ChaosVehicles: fix steering jitter

diff --git a/Engine/Plugins/Experimental/ChaosVehiclesPlugin/Source/ChaosVehicles/Private/ChaosWheeledVehicleMovementComponent.cpp b/Engine/Plugins/Experimental/ChaosVehiclesPlugin/Source/ChaosVehicles/Private/ChaosWheeledVehicleMovementComponent.cpp
--- a/Engine/Plugins/Experimental/ChaosVehiclesPlugin/Source/ChaosVehicles/Private/ChaosWheeledVehicleMovementComponent.cpp
+++ b/Engine/Plugins/Experimental/ChaosVehiclesPlugin/Source/ChaosVehicles/Private/ChaosWheeledVehicleMovementComponent.cpp
@@ -893,7 +893,7 @@ void UChaosWheeledVehicleMovementComponent::ProcessSteering()
 			float SpeedScale = 1.0f;
 
 			// allow full counter steering when steering into a power slide
-			if (SteeringInput * VehicleState.LocalGForce.Y < 0.1f)
+			//if (SteeringInput * VehicleState.VehicleLocalVelocity.Y > 0.0f)
 			{
 				SpeedScale = SteeringCurveData->Eval(CmSToMPH(VehicleState.ForwardSpeed));
 			}
