ChaosVehicles: Fix issue with raycasts against trigger zones

diff --git a/Engine/Plugins/Experimental/ChaosVehiclesPlugin/Source/ChaosVehicles/Private/ChaosWheeledVehicleMovementComponent.cpp b/Engine/Plugins/Experimental/ChaosVehiclesPlugin/Source/ChaosVehicles/Private/ChaosWheeledVehicleMovementComponent.cpp
--- a/Engine/Plugins/Experimental/ChaosVehiclesPlugin/Source/ChaosVehicles/Private/ChaosWheeledVehicleMovementComponent.cpp
+++ b/Engine/Plugins/Experimental/ChaosVehiclesPlugin/Source/ChaosVehicles/Private/ChaosWheeledVehicleMovementComponent.cpp
@@ -679,6 +679,9 @@ void UChaosWheeledVehicleMovementComponent::PerformSuspensionTraces(const TArray
 				// Test each overlapped object for a hit result
 				for (FOverlapResult OverlapResult : OverlapResults)
 				{
+					if (!OverlapResult.bBlockingHit)
+						continue;
+
 					FHitResult ComponentHit;
 
 					switch (Wheels[WheelIdx]->SweepShape)
