Fix Chaos Vehicles going to sleep on Servers and Replicated Clients due to no input.

diff --git a/Engine/Plugins/Experimental/ChaosVehiclesPlugin/Source/ChaosVehicles/Private/ChaosVehicleMovementComponent.cpp b/Engine/Plugins/Experimental/ChaosVehiclesPlugin/Source/ChaosVehicles/Private/ChaosVehicleMovementComponent.cpp
--- a/Engine/Plugins/Experimental/ChaosVehiclesPlugin/Source/ChaosVehicles/Private/ChaosVehicleMovementComponent.cpp
+++ b/Engine/Plugins/Experimental/ChaosVehiclesPlugin/Source/ChaosVehicles/Private/ChaosVehicleMovementComponent.cpp
@@ -942,16 +942,28 @@ void UChaosVehicleMovementComponent::ProcessSleeping()
 			VehicleState.SleepCounter = 0;
 		}
 
-		bool ControlInputPressed = (RawThrottleInput >= SMALL_NUMBER) || (RawBrakeInput >= SMALL_NUMBER) || (FMath::Abs(RawSteeringInput) > SMALL_NUMBER);
+		// If the vehicle is locally controlled, we want to use the raw inputs to determine sleep.
+		// However, if it's on the Server or is just being replicated to other Clients then there
+		// won't be any Raw input. In that case, use ReplicatedState instead.
+		
+		// NOTE: Even on local clients, ReplicatedState will still be populated (the call to ServerUpdateState will
+		//			be processed locally). Maybe we should *just* use ReplicatedState?
+
+		// TODO: What about other inputs, like handbrake, roll, pitch, yaw?
+		const AController* Controller = GetController();
+		const bool bIsLocallyControlled = (Controller && Controller->IsLocalController());
+		const bool bControlInputPressed = bIsLocallyControlled ?
+			(RawThrottleInput >= SMALL_NUMBER) || (RawBrakeInput >= SMALL_NUMBER) || (FMath::Abs(RawSteeringInput) > SMALL_NUMBER) :
+			(ReplicatedState.ThrottleInput >= SMALL_NUMBER) || (ReplicatedState.BrakeInput >= SMALL_NUMBER) || (FMath::Abs(ReplicatedState.SteeringInput) > SMALL_NUMBER);
 
 		// Wake if control input pressed
-		if (VehicleState.bSleeping && (ControlInputPressed || !VehicleState.bAllWheelsOnGround))
+		if (VehicleState.bSleeping && (bControlInputPressed || !VehicleState.bAllWheelsOnGround))
 		{
 			VehicleState.bSleeping = false;
 			VehicleState.SleepCounter = 0;
 			TargetInstance->WakeInstance();
 		}
-		else if (!VehicleState.bSleeping && !ControlInputPressed && VehicleState.bAllWheelsOnGround && (VehicleState.VehicleUpAxis.Z > SleepSlopeLimit))
+		else if (!VehicleState.bSleeping && !bControlInputPressed && VehicleState.bAllWheelsOnGround && (VehicleState.VehicleUpAxis.Z > SleepSlopeLimit))
 		{
 			float SpeedSqr = TargetInstance->GetUnrealWorldVelocity().SizeSquared();
 			if (SpeedSqr < (SleepThreshold* SleepThreshold))
