ChaosVehicles: Fix networking issue when clearing input

diff --git a/Engine/Plugins/Experimental/ChaosVehiclesPlugin/Source/ChaosVehicles/Private/ChaosVehicleMovementComponent.cpp b/Engine/Plugins/Experimental/ChaosVehiclesPlugin/Source/ChaosVehicles/Private/ChaosVehicleMovementComponent.cpp
--- a/Engine/Plugins/Experimental/ChaosVehiclesPlugin/Source/ChaosVehicles/Private/ChaosVehicleMovementComponent.cpp
+++ b/Engine/Plugins/Experimental/ChaosVehiclesPlugin/Source/ChaosVehicles/Private/ChaosVehicleMovementComponent.cpp
@@ -560,7 +560,11 @@ void UChaosVehicleMovementComponent::ClearInput()
 		CurrentGear = PVehicle->GetTransmission().GetCurrentGear();
 	}
 
-	ServerUpdateState(SteeringInput, ThrottleInput, BrakeInput, HandbrakeInput, CurrentGear, RollInput, PitchInput, YawInput);
+	AController* Controller = GetController();
+	if (Controller && Controller->IsLocalController() && PVehicle)
+	{
+		ServerUpdateState(SteeringInput, ThrottleInput, BrakeInput, HandbrakeInput, CurrentGear, RollInput, PitchInput, YawInput);
+	}
 }
 
 void UChaosVehicleMovementComponent::ClearRawInput()
