ChaosVehicles: Fix StopMovementImmediately functionality for vehicles

diff --git a/Engine/Plugins/Experimental/ChaosVehiclesPlugin/Source/ChaosVehicles/Private/ChaosVehicleMovementComponent.cpp b/Engine/Plugins/Experimental/ChaosVehiclesPlugin/Source/ChaosVehicles/Private/ChaosVehicleMovementComponent.cpp
--- a/Engine/Plugins/Experimental/ChaosVehiclesPlugin/Source/ChaosVehicles/Private/ChaosVehicleMovementComponent.cpp
+++ b/Engine/Plugins/Experimental/ChaosVehiclesPlugin/Source/ChaosVehicles/Private/ChaosVehicleMovementComponent.cpp
@@ -309,6 +309,14 @@ void UChaosVehicleMovementComponent::TickVehicle(float DeltaTime)
 
 void UChaosVehicleMovementComponent::StopMovementImmediately()
 {
+	FBodyInstance* TargetInstance = GetBodyInstance();
+	if (TargetInstance)
+	{
+		TargetInstance->SetLinearVelocity(FVector::ZeroVector, false);
+		TargetInstance->SetAngularVelocityInRadians(FVector::ZeroVector, false);
+		TargetInstance->ClearForces();
+		TargetInstance->ClearTorques();
+	}
 	Super::StopMovementImmediately();
 	ClearAllInput();
 }
