ChaosVehicles: VehiclePlugin module type setup incorrectly

diff --git a/Engine/Plugins/Experimental/ChaosVehiclesPlugin/ChaosVehiclesPlugin.uplugin b/Engine/Plugins/Experimental/ChaosVehiclesPlugin/ChaosVehiclesPlugin.uplugin
--- a/Engine/Plugins/Experimental/ChaosVehiclesPlugin/ChaosVehiclesPlugin.uplugin
+++ b/Engine/Plugins/Experimental/ChaosVehiclesPlugin/ChaosVehiclesPlugin.uplugin
@@ -23,7 +23,7 @@
 		},
 		{
 			"Name": "ChaosVehiclesEditor",
-			"Type": "Editor",
+			"Type": "UncookedOnly",
 			"LoadingPhase": "PreDefault"
 		}
 	]
