ChaosVehicles: Expose whether engine is valid/enabled for a vehicle, useful when Mechanical Simulation is not enabled & no engine present

diff --git a/Engine/Plugins/Experimental/ChaosVehiclesPlugin/Source/ChaosVehicles/Private/ChaosWheeledVehicleMovementComponent.cpp b/Engine/Plugins/Experimental/ChaosVehiclesPlugin/Source/ChaosVehicles/Private/ChaosWheeledVehicleMovementComponent.cpp
--- a/Engine/Plugins/Experimental/ChaosVehiclesPlugin/Source/ChaosVehicles/Private/ChaosWheeledVehicleMovementComponent.cpp
+++ b/Engine/Plugins/Experimental/ChaosVehiclesPlugin/Source/ChaosVehicles/Private/ChaosWheeledVehicleMovementComponent.cpp
@@ -1133,7 +1133,7 @@ float UChaosWheeledVehicleMovementComponent::GetEngineRotationSpeed() const
 {
 	float EngineRPM = 0.f;
 
-	if (PVehicle.IsValid())
+	if (PVehicle.IsValid() && PVehicle->HasEngine())
 	{
 		EngineRPM = PVehicle->GetEngine().GetEngineRPM();
 	}
@@ -1145,7 +1145,7 @@ float UChaosWheeledVehicleMovementComponent::GetEngineMaxRotationSpeed() const
 {
 	float MaxEngineRPM = 0.f;
 	
-	if (PVehicle.IsValid())
+	if (PVehicle.IsValid() && PVehicle->HasEngine())
 	{
 		MaxEngineRPM = PVehicle->GetEngine().Setup().MaxRPM;
 	}
diff --git a/Engine/Plugins/Experimental/ChaosVehiclesPlugin/Source/ChaosVehicles/Public/ChaosVehicleMovementComponent.h b/Engine/Plugins/Experimental/ChaosVehiclesPlugin/Source/ChaosVehicles/Public/ChaosVehicleMovementComponent.h
--- a/Engine/Plugins/Experimental/ChaosVehiclesPlugin/Source/ChaosVehicles/Public/ChaosVehicleMovementComponent.h
+++ b/Engine/Plugins/Experimental/ChaosVehiclesPlugin/Source/ChaosVehicles/Public/ChaosVehicleMovementComponent.h
@@ -698,7 +698,6 @@ public:
 
 	TUniquePtr<Chaos::FSimpleWheeledVehicle>& PhysicsVehicle()
 	{
-		check(PVehicle.IsValid());
 		return PVehicle;
 	}
 
diff --git a/Engine/Source/Runtime/Experimental/ChaosVehicles/ChaosVehiclesCore/Public/SimpleVehicle.h b/Engine/Source/Runtime/Experimental/ChaosVehicles/ChaosVehiclesCore/Public/SimpleVehicle.h
--- a/Engine/Source/Runtime/Experimental/ChaosVehicles/ChaosVehiclesCore/Public/SimpleVehicle.h
+++ b/Engine/Source/Runtime/Experimental/ChaosVehicles/ChaosVehiclesCore/Public/SimpleVehicle.h
@@ -59,6 +59,11 @@ public:
 		return Engine[0];
 	}
 
+	bool HasEngine() const
+	{
+		return (Engine.Num() > 0);
+	}
+
 	bool HasTransmission() const
 	{
 		return (Transmission.Num() > 0);
