Chaos : Vehicle
-Relaxed condition in UChaosVehicleMovementComponent::LocateBoneOffset to issue a message instead of exiting the application.

diff --git a/Engine/Plugins/Experimental/ChaosVehiclesPlugin/Source/ChaosVehicles/Private/ChaosVehicleMovementComponent.cpp b/Engine/Plugins/Experimental/ChaosVehiclesPlugin/Source/ChaosVehicles/Private/ChaosVehicleMovementComponent.cpp
--- a/Engine/Plugins/Experimental/ChaosVehiclesPlugin/Source/ChaosVehicles/Private/ChaosVehicleMovementComponent.cpp
+++ b/Engine/Plugins/Experimental/ChaosVehiclesPlugin/Source/ChaosVehicles/Private/ChaosVehicleMovementComponent.cpp
@@ -1080,17 +1080,19 @@ FVector UChaosVehicleMovementComponent::LocateBoneOffset(const FName InBoneName,
 	{
 		if (USkinnedMeshComponent* Mesh = Cast<USkinnedMeshComponent>(GetMesh()))
 		{
-			check(Mesh->SkeletalMesh);
-			const FVector BonePosition = Mesh->SkeletalMesh->GetComposedRefPoseMatrix(InBoneName).GetOrigin() * Mesh->GetRelativeScale3D();
-			//BonePosition is local for the root BONE of the skeletal mesh - however, we are using the Root BODY which may have its own transform, so we need to return the position local to the root BODY
-			FMatrix RootBodyMTX = FMatrix::Identity;
-			// Body Instance is no longer valid at this point in the code
-			if (Mesh->GetBodyInstance())
+			if (ensureMsgf(Mesh->SkeletalMesh, TEXT("Expected skeletal mesh when locating bone offset. Asset might be missing references.")))
 			{
-				RootBodyMTX = Mesh->SkeletalMesh->GetComposedRefPoseMatrix(Mesh->GetBodyInstance()->BodySetup->BoneName);
+				const FVector BonePosition = Mesh->SkeletalMesh->GetComposedRefPoseMatrix(InBoneName).GetOrigin() * Mesh->GetRelativeScale3D();
+				//BonePosition is local for the root BONE of the skeletal mesh - however, we are using the Root BODY which may have its own transform, so we need to return the position local to the root BODY
+				FMatrix RootBodyMTX = FMatrix::Identity;
+				// Body Instance is no longer valid at this point in the code
+				if (Mesh->GetBodyInstance())
+				{
+					RootBodyMTX = Mesh->SkeletalMesh->GetComposedRefPoseMatrix(Mesh->GetBodyInstance()->BodySetup->BoneName);
+				}
+				const FVector LocalBonePosition = RootBodyMTX.InverseTransformPosition(BonePosition);
+				Offset += LocalBonePosition;
 			}
-			const FVector LocalBonePosition = RootBodyMTX.InverseTransformPosition(BonePosition);
-			Offset += LocalBonePosition;
 		}
 	}
 	return Offset;
