ChaosVehicles: Major reworking to support all simulation code to run off the async physics callback.
Rework so game code only touches game side setup data and recieves output through controlled output mechanism from async thread.

diff --git a/Engine/Plugins/Experimental/ChaosVehiclesPlugin/Source/ChaosVehicles/Private/AsyncCallback.cpp b/Engine/Plugins/Experimental/ChaosVehiclesPlugin/Source/ChaosVehicles/Private/AsyncCallback.cpp
new file mode 100644
--- /dev/null
+++ b/Engine/Plugins/Experimental/ChaosVehiclesPlugin/Source/ChaosVehicles/Private/AsyncCallback.cpp
@@ -0,0 +1,85 @@
+// Copyright Epic Games, Inc. All Rights Reserved.
+
+#include "AsyncCallback.h"
+
+#include "ChaosVehicleMovementComponent.h"
+#include "PBDRigidsSolver.h"
+#include "Chaos/ParticleHandleFwd.h"
+#include "PhysicsProxy/SingleParticlePhysicsProxy.h"
+
+extern FVehicleDebugParams GVehicleDebugParams;
+
+/**
+ * Callback from Physics thread
+ */
+void FChaosVehicleManagerAsyncCallback::OnPreSimulate_Internal()
+{
+	using namespace Chaos;
+
+	float DeltaTime = GetDeltaTime_Internal();
+	float SimTime = GetSimTime_Internal();
+
+	const FChaosVehicleManagerAsyncInput* Input = GetConsumerInput_Internal();
+	if (Input == nullptr)
+	{
+		return;
+	}
+
+	const int32 NumVehicles = Input->VehicleInputs.Num();
+
+	UWorld* World = Input->World.Get();	//only safe to access for scene queries
+	if (World == nullptr || NumVehicles == 0)
+	{
+		//world is gone so don't bother, or nothing to simulate.
+		return;
+	}
+
+	Chaos::FPhysicsSolver* PhysicsSolver = static_cast<Chaos::FPhysicsSolver*>(GetSolver());
+	if (PhysicsSolver == nullptr)
+	{
+		return;
+	}
+
+	FChaosVehicleManagerAsyncOutput& Output = GetProducerOutputData_Internal();
+	Output.VehicleOutputs.AddDefaulted(NumVehicles);
+	Output.Timestamp = Input->Timestamp;
+
+	const auto& InputVehiclesBatch = Input->VehicleInputs;
+	auto& OutputVehiclesBatch = Output.VehicleOutputs;
+
+	// beware running the vehicle simulation in parallel, code must remain threadsafe
+	auto LambdaParallelUpdate = [World, DeltaTime, SimTime, &InputVehiclesBatch, &OutputVehiclesBatch](int32 Idx)
+	{
+		const FChaosVehicleAsyncInput& VehicleInput = *InputVehiclesBatch[Idx];
+
+		if (VehicleInput.Actor.Proxy == nullptr || VehicleInput.Actor.Proxy->GetHandle() == nullptr)
+		{
+			return;
+		}
+
+		auto Handle = VehicleInput.Actor.Proxy->GetHandle();
+		if (Handle->ObjectState() != Chaos::EObjectStateType::Dynamic)
+		{
+			UE_LOG(LogChaos, Warning, TEXT("Handle->ObjectState() %d"), Handle->ObjectState());
+			return;
+		}
+
+		bool bWake = false;
+		OutputVehiclesBatch[Idx] = VehicleInput.Simulate(World, DeltaTime, SimTime, bWake);
+
+	};
+
+	bool ForceSingleThread = !GVehicleDebugParams.EnableMultithreading;
+	ParallelFor(OutputVehiclesBatch.Num(), LambdaParallelUpdate, ForceSingleThread);
+
+
+}
+
+/**
+ * Contact modification currently unused
+ */
+void FChaosVehicleManagerAsyncCallback::OnContactModification_Internal(const TArrayView<Chaos::FPBDCollisionConstraintHandleModification>& Modifications)
+{
+
+}
+
diff --git a/Engine/Plugins/Experimental/ChaosVehiclesPlugin/Source/ChaosVehicles/Private/ChaosVehicleManager.cpp b/Engine/Plugins/Experimental/ChaosVehiclesPlugin/Source/ChaosVehicles/Private/ChaosVehicleManager.cpp
--- a/Engine/Plugins/Experimental/ChaosVehiclesPlugin/Source/ChaosVehicles/Private/ChaosVehicleManager.cpp
+++ b/Engine/Plugins/Experimental/ChaosVehiclesPlugin/Source/ChaosVehicles/Private/ChaosVehicleManager.cpp
@@ -6,6 +6,10 @@
 #include "PhysicalMaterials/PhysicalMaterial.h"
 #include "Physics/PhysicsFiltering.h"
 #include "Physics/PhysicsInterfaceCore.h"
+#include "ChaosVehicleMovementComponent.h"
+
+#include "PBDRigidsSolver.h"
+#include "PhysicsProxy/SingleParticlePhysicsProxy.h"
 
 
 DECLARE_STATS_GROUP(TEXT("ChaosVehicleManager"), STATGROUP_ChaosVehicleManager, STATGROUP_Advanced);
@@ -15,36 +19,102 @@ DECLARE_CYCLE_STAT(TEXT("TickVehicles"), STAT_ChaosVehicleManager_TickVehicles,
 DECLARE_CYCLE_STAT(TEXT("VehicleManagerUpdate"), STAT_ChaosVehicleManager_Update, STATGROUP_ChaosVehicleManager);
 DECLARE_CYCLE_STAT(TEXT("PretickVehicles"), STAT_ChaosVehicleManager_PretickVehicles, STATGROUP_Physics);
 
+extern FVehicleDebugParams GVehicleDebugParams;
 
 TMap<FPhysScene*, FChaosVehicleManager*> FChaosVehicleManager::SceneToVehicleManagerMap;
 uint32 FChaosVehicleManager::VehicleSetupTag = 0;
 
+FDelegateHandle FChaosVehicleManager::OnPostWorldInitializationHandle;
+FDelegateHandle FChaosVehicleManager::OnWorldCleanupHandle;
+bool FChaosVehicleManager::GInitialized = false;
+
 FChaosVehicleManager::FChaosVehicleManager(FPhysScene* PhysScene)
 #if WITH_CHAOS
 	: Scene(*PhysScene)
+	, AsyncCallback(nullptr)
+	, Timestamp(0)
+
 #endif
 {
-	// Set up delegates
-	OnPhysScenePreTickHandle = PhysScene->OnPhysScenePreTick.AddRaw(this, &FChaosVehicleManager::PreTick);
-	OnPhysSceneStepHandle = PhysScene->OnPhysSceneStep.AddRaw(this, &FChaosVehicleManager::Update);
+	check(PhysScene);
+	
+	if (!GInitialized)
+	{
+		GInitialized = true;
+		// PhysScene->GetOwningWorld() is always null here, the world is being setup too late to be of use
+		// therefore setup these global world delegates that will callback when everything is setup so registering
+		// the physics solver Async Callback will succeed
+		OnPostWorldInitializationHandle = FWorldDelegates::OnPostWorldInitialization.AddStatic(&FChaosVehicleManager::OnPostWorldInitialization);
+		OnWorldCleanupHandle = FWorldDelegates::OnWorldCleanup.AddStatic(&FChaosVehicleManager::OnWorldCleanup);
+	}
 
 	// Add to Scene-To-Manager map
 	FChaosVehicleManager::SceneToVehicleManagerMap.Add(PhysScene, this);
 }
 
+
+void FChaosVehicleManager::RegisterCallbacks()
+{
+	OnPhysScenePreTickHandle = Scene.OnPhysScenePreTick.AddRaw(this, &FChaosVehicleManager::Update);
+	OnPhysSceneStepHandle = Scene.OnPhysSceneStep.AddRaw(this, &FChaosVehicleManager::SubStep);
+	OnPhysScenePostTickHandle = Scene.OnPhysScenePostTick.AddRaw(this, &FChaosVehicleManager::PostUpdate);
+
+	check(AsyncCallback == nullptr);
+	AsyncCallback = Scene.GetSolver()->CreateAndRegisterSimCallbackObject_External<FChaosVehicleManagerAsyncCallback>(true);
+}
+
+void FChaosVehicleManager::UnregisterCallbacks()
+{
+	Scene.OnPhysScenePreTick.Remove(OnPhysScenePreTickHandle);
+	Scene.OnPhysSceneStep.Remove(OnPhysSceneStepHandle);
+	Scene.OnPhysScenePostTick.Remove(OnPhysScenePostTickHandle);
+
+	if (AsyncCallback)
+	{
+		Scene.GetSolver()->UnregisterAndFreeSimCallbackObject_External(AsyncCallback);
+		AsyncCallback = nullptr;
+	}
+}
+
+void FChaosVehicleManager::OnPostWorldInitialization(UWorld* InWorld, const UWorld::InitializationValues)
+{
+	FChaosVehicleManager* Manager = FChaosVehicleManager::GetVehicleManagerFromScene(InWorld->GetPhysicsScene());
+	if (Manager)
+	{
+		Manager->RegisterCallbacks();
+	}
+}
+
+
+void FChaosVehicleManager::OnWorldCleanup(UWorld* InWorld, bool bSessionEnded, bool bCleanupResources)
+{
+	FChaosVehicleManager* Manager = FChaosVehicleManager::GetVehicleManagerFromScene(InWorld->GetPhysicsScene());
+	if (Manager)
+	{
+		Manager->UnregisterCallbacks();
+	}
+}
+
 void FChaosVehicleManager::DetachFromPhysScene(FPhysScene* PhysScene)
 {
-	PhysScene->OnPhysScenePreTick.Remove(OnPhysScenePreTickHandle);
-	PhysScene->OnPhysSceneStep.Remove(OnPhysSceneStepHandle);
+	if (AsyncCallback)
+	{
+		UnregisterCallbacks();
+	}
+
+	if (PhysScene->GetOwningWorld())
+	{
+		PhysScene->GetOwningWorld()->OnWorldBeginPlay.RemoveAll(this);
+	}
 
 	FChaosVehicleManager::SceneToVehicleManagerMap.Remove(PhysScene);
 }
 
 FChaosVehicleManager::~FChaosVehicleManager()
 {
-	while( Vehicles.Num() > 0 )
+	while (Vehicles.Num() > 0)
 	{
-		RemoveVehicle( Vehicles.Last() );
+		RemoveVehicle(Vehicles.Last());
 	}
 }
 
@@ -59,61 +129,192 @@ FChaosVehicleManager* FChaosVehicleManager::GetVehicleManagerFromScene(FPhysScen
 	return Manager;
 }
 
-void FChaosVehicleManager::AddVehicle( TWeakObjectPtr<UChaosVehicleMovementComponent> Vehicle )
+void FChaosVehicleManager::AddVehicle(TWeakObjectPtr<UChaosVehicleMovementComponent> Vehicle)
 {
 	check(Vehicle != NULL);
-	check(Vehicle->PhysicsVehicle());
+	check(Vehicle->PhysicsVehicleOutput());
+	check(AsyncCallback);
 
-	Vehicles.Add( Vehicle );
+	Vehicles.Add(Vehicle);
 }
 
-void FChaosVehicleManager::RemoveVehicle( TWeakObjectPtr<UChaosVehicleMovementComponent> Vehicle )
+void FChaosVehicleManager::RemoveVehicle(TWeakObjectPtr<UChaosVehicleMovementComponent> Vehicle)
 {
 	check(Vehicle != NULL);
-	check(Vehicle->PhysicsVehicle());
+	check(Vehicle->PhysicsVehicleOutput());
 
-	int32 RemovedIndex = Vehicles.Find(Vehicle);
+	Vehicles.Remove(Vehicle);
 
-	Vehicles.Remove( Vehicle );
+	if (Vehicle->PhysicsVehicleOutput().IsValid())
+	{
+		Vehicle->PhysicsVehicleOutput().Reset(nullptr);
+	}
 
-	if (Vehicle->PhysicsVehicle().IsValid())
+}
+
+void FChaosVehicleManager::ScenePreTick(FPhysScene* PhysScene, float DeltaTime)
+{
+	// inputs being set via back door, i.e. accessing PVehicle directly is a no go now, needs to go through async input system
+	SCOPE_CYCLE_COUNTER(STAT_ChaosVehicleManager_PretickVehicles);
+
+	for (int32 i = 0; i < Vehicles.Num(); ++i)
 	{
-		Vehicle->PhysicsVehicle().Reset(nullptr);
+		Vehicles[i]->PreTickGT(DeltaTime);
 	}
+
 }
 
 void FChaosVehicleManager::Update(FPhysScene* PhysScene, float DeltaTime)
 {
-	SCOPE_CYCLE_COUNTER(STAT_ChaosVehicleManager_Update);
+	UWorld* World = Scene.GetOwningWorld();
+
+	SubStepCount = 0;
+
+	ScenePreTick(PhysScene, DeltaTime);
+
+	ParallelUpdateVehicles(DeltaTime);
+
+	if (World)
+	{
+		FChaosVehicleManagerAsyncInput* AsyncInput = AsyncCallback->GetProducerInputData_External();
+		for (TWeakObjectPtr<UChaosVehicleMovementComponent> Vehicle : Vehicles)
+		{
+			Vehicle->Update(DeltaTime);
+			Vehicle->FinalizeSimCallbackData(*AsyncInput);
+		}
+	}
+}
+
+void FChaosVehicleManager::SubStep(FPhysScene* PhysScene, float DeltaTime)
+{
+	if (SubStepCount++ > 0)	//ignore first substep
+	{
+		ParallelUpdateVehicles(DeltaTime);
+		//if (UWorld* World = GetWorld())
+		//{
+		//	for (TWeakObjectPtr<UChaosVehicleMovementComponent> Vehicle : Vehicles)
+		//	{
+		//		if (UPrimitiveComponent* RootPrim = Cast<UPrimitiveComponent>(Vehicle->GetRootComponent()))
+		//		{
+		//			if (RootPrim->GetBodyInstance() && !Vehicle->IsAsleep())
+		//			{
+		//				Vehicle->Substep(DeltaTime);
+		//			}
+		//		}
+		//	}
+		//}
+	}
+}
+
+void FChaosVehicleManager::PostUpdate(FChaosScene* PhysScene)
+{
+	// Unused
+}
+
+void FChaosVehicleManager::ParallelUpdateVehicles(float DeltaSeconds)
+{
+	FChaosVehicleManagerAsyncInput* AsyncInput = AsyncCallback->GetProducerInputData_External();
+
+	AsyncInput->Reset();	//only want latest frame's data
 
-	if (Vehicles.Num() == 0 )
 	{
-		return;
+		// We pass pointers from TArray so this reserve is critical. Otherwise realloc happens
+		AsyncInput->VehicleInputs.Reserve(Vehicles.Num());
+		AsyncInput->Timestamp = Timestamp;
+		AsyncInput->World = Scene.GetOwningWorld();
 	}
 
-	//	Suspension raycasts
+	// Grab all outputs for processing, even future ones for interpolation.
 	{
-		//SCOPE_CYCLE_COUNTER(STAT_ChaosVehicleManager_VehicleSuspensionRaycasts);
-		// possibly batch all the vehicle raycasts?
+		Chaos::TSimCallbackOutputHandle<FChaosVehicleManagerAsyncOutput> AsyncOutputLatest;
+		while ((AsyncOutputLatest = AsyncCallback->PopFutureOutputData_External()))
+		{
+			PendingOutputs.Emplace(MoveTemp(AsyncOutputLatest));
+		}
 	}
-	
-	// Tick vehicles
+
+	// Since we are in pre-physics, delta seconds is not accounted for in external time yet
+	const float ResultsTime = AsyncCallback->GetSolver()->GetPhysicsResultsTime_External() + DeltaSeconds;
+
+	// Find index of first non-consumable output (first one after current time)
+	int32 LastOutputIdx = 0;
+	for (; LastOutputIdx < PendingOutputs.Num(); ++LastOutputIdx)
 	{
-		SCOPE_CYCLE_COUNTER(STAT_ChaosVehicleManager_TickVehicles);
-		for (int32 i = Vehicles.Num() - 1; i >= 0; --i)
+		if (PendingOutputs[LastOutputIdx]->InternalTime > ResultsTime)
 		{
-			Vehicles[i]->TickVehicle(DeltaTime);
+			break;
 		}
 	}
 
-}
+	// Process events on all outputs which occurred before current time
+	// 
+	//for (int32 OutputIdx = 0; OutputIdx < LastOutputIdx; ++OutputIdx)
+	//{
+	//	for (TWeakObjectPtr<UChaosVehicleMovementComponent> Vehicle : Vehicles)
+	//	{
+	//		Vehicle->GameThread_ProcessIntermediateAsyncOutput(*PendingOutputs[OutputIdx]);
+	//	}
+	//}
 
-void FChaosVehicleManager::PreTick(FPhysScene* PhysScene, float DeltaTime)
-{
-	SCOPE_CYCLE_COUNTER(STAT_ChaosVehicleManager_PretickVehicles);
+	// Cache the last consumed output for interpolation
+	if (LastOutputIdx > 0)
+	{
+		LatestOutput = MoveTemp(PendingOutputs[LastOutputIdx - 1]);
+	}
 
-	for (int32 i = 0; i < Vehicles.Num(); ++i)
+	// Remove all consumed outputs
+	{
+		TArray<Chaos::TSimCallbackOutputHandle<FChaosVehicleManagerAsyncOutput>> NewPendingOutputs;
+		for (int32 OutputIdx = LastOutputIdx; OutputIdx < PendingOutputs.Num(); ++OutputIdx)
+		{
+			NewPendingOutputs.Emplace(MoveTemp(PendingOutputs[OutputIdx]));
+		}
+		PendingOutputs = MoveTemp(NewPendingOutputs);
+	}
+
+	// It's possible we will end up multiple frames ahead of output, take the latest ready output.
+	Chaos::TSimCallbackOutputHandle<FChaosVehicleManagerAsyncOutput> AsyncOutput;
+	Chaos::TSimCallbackOutputHandle<FChaosVehicleManagerAsyncOutput> AsyncOutputLatest;
+	while ((AsyncOutputLatest = AsyncCallback->PopOutputData_External()))
 	{
-		Vehicles[i]->PreTick(DeltaTime);
+		AsyncOutput = MoveTemp(AsyncOutputLatest);
+
+		for (TWeakObjectPtr<UChaosVehicleMovementComponent> Vehicle : Vehicles)
+		{
+			//Vehicle->GameThread_ProcessIntermediateAsyncOutput(*AsyncOutput);
+		}
 	}
+
+	if (UWorld* World = Scene.GetOwningWorld())
+	{
+		int32 NumVehiclesInActiveBatch = 0;
+		for (TWeakObjectPtr<UChaosVehicleMovementComponent> Vehicle : Vehicles)
+		{
+			auto NextOutput = PendingOutputs.Num() > 0 ? PendingOutputs[0].Get() : nullptr;
+			float Alpha = 0.f;
+			if (NextOutput && LatestOutput)
+			{
+				const float Denom = NextOutput->InternalTime - LatestOutput->InternalTime;
+				if (Denom > SMALL_NUMBER)
+				{
+					Alpha = (ResultsTime - LatestOutput->InternalTime) / Denom;
+				}
+			}
+
+			AsyncInput->VehicleInputs.Add(Vehicle->SetCurrentAsyncInputOutput(AsyncInput->VehicleInputs.Num(), LatestOutput.Get(), NextOutput, Alpha, Timestamp));
+		}
+	}
+
+	++Timestamp;
+
+	const auto& AwakeVehiclesBatch = Vehicles; // TODO: process awake only
+
+	auto LambdaParallelUpdate = [DeltaSeconds, &AwakeVehiclesBatch](int32 Idx)
+	{
+		TWeakObjectPtr<UChaosVehicleMovementComponent> Vehicle = AwakeVehiclesBatch[Idx];
+		Vehicle->ParallelUpdate(DeltaSeconds); // gets output state from PT
+	};
+
+	bool ForceSingleThread = !GVehicleDebugParams.EnableMultithreading;
+	ParallelFor(AwakeVehiclesBatch.Num(), LambdaParallelUpdate, ForceSingleThread);
 }
diff --git a/Engine/Plugins/Experimental/ChaosVehiclesPlugin/Source/ChaosVehicles/Private/ChaosVehicleMovementComponent.cpp b/Engine/Plugins/Experimental/ChaosVehiclesPlugin/Source/ChaosVehicles/Private/ChaosVehicleMovementComponent.cpp
--- a/Engine/Plugins/Experimental/ChaosVehiclesPlugin/Source/ChaosVehicles/Private/ChaosVehicleMovementComponent.cpp
+++ b/Engine/Plugins/Experimental/ChaosVehiclesPlugin/Source/ChaosVehicles/Private/ChaosVehicleMovementComponent.cpp
@@ -31,6 +31,11 @@
 #include "PhysicalMaterials/PhysicalMaterial.h"
 #include "GameFramework/HUD.h"
 
+#include "Chaos/Particle/ParticleUtilities.h"
+#include "Chaos/ParticleHandleFwd.h"
+#include "PhysicsProxy/SingleParticlePhysicsProxy.h"
+
+
 #define LOCTEXT_NAMESPACE "UVehicleMovementComponent"
 
 #if VEHICLE_DEBUGGING_ENABLED
@@ -59,13 +64,10 @@
 FAutoConsoleVariableRef CVarChaosVehiclesDisableVehicleSleep(TEXT("p.Vehicle.DisableVehicleSleep"), GVehicleDebugParams.DisableVehicleSleep, TEXT("Disable Vehicle Agressive Sleeping."));
 
 
-void FVehicleState::CaptureState(FBodyInstance* TargetInstance, float GravityZ, float DeltaTime)
+void FVehicleState::CaptureState(const FBodyInstance* TargetInstance, float GravityZ, float DeltaTime)
 {
 	if (TargetInstance)
 	{
-		VehicleUpAxis = VehicleWorldTransform.GetUnitAxis(EAxis::Z);
-		VehicleForwardAxis = VehicleWorldTransform.GetUnitAxis(EAxis::X);
-		VehicleRightAxis = VehicleWorldTransform.GetUnitAxis(EAxis::Y);
 
 		VehicleWorldTransform = TargetInstance->GetUnrealWorldTransform();
 		VehicleWorldVelocity = TargetInstance->GetUnrealWorldVelocity();
@@ -73,6 +75,10 @@
 		VehicleWorldCOM = TargetInstance->GetCOMPosition();
 		WorldVelocityNormal = VehicleWorldVelocity.GetSafeNormal();
 
+		VehicleUpAxis = VehicleWorldTransform.GetUnitAxis(EAxis::Z);
+		VehicleForwardAxis = VehicleWorldTransform.GetUnitAxis(EAxis::X);
+		VehicleRightAxis = VehicleWorldTransform.GetUnitAxis(EAxis::Y);
+
 		VehicleLocalVelocity = VehicleWorldTransform.InverseTransformVector(VehicleWorldVelocity);
 		LocalAcceleration = (VehicleLocalVelocity - LastFrameVehicleLocalVelocity) / DeltaTime;
 		LocalGForce = LocalAcceleration / FMath::Abs(GravityZ);
@@ -83,6 +89,504 @@
 	}
 }
 
+void FVehicleState::CaptureState(const Chaos::TPBDRigidParticleHandle<float, 3>* Handle, float GravityZ, float DeltaTime)
+{
+	if (Handle)
+	{
+		const FTransform WorldTM(Handle->R(), Handle->X());
+		VehicleWorldTransform = WorldTM;
+		VehicleWorldVelocity = Handle->V();
+		VehicleWorldAngularVelocity = Handle->W();
+		VehicleWorldCOM = Handle->CenterOfMass();
+		WorldVelocityNormal = VehicleWorldVelocity.GetSafeNormal();
+
+		VehicleUpAxis = VehicleWorldTransform.GetUnitAxis(EAxis::Z);
+		VehicleForwardAxis = VehicleWorldTransform.GetUnitAxis(EAxis::X);
+		VehicleRightAxis = VehicleWorldTransform.GetUnitAxis(EAxis::Y);
+
+		VehicleLocalVelocity = VehicleWorldTransform.InverseTransformVector(VehicleWorldVelocity);
+		LocalAcceleration = (VehicleLocalVelocity - LastFrameVehicleLocalVelocity) / DeltaTime;
+		LocalGForce = LocalAcceleration / FMath::Abs(GravityZ);
+		LastFrameVehicleLocalVelocity = VehicleLocalVelocity;
+
+		ForwardSpeed = FVector::DotProduct(VehicleWorldVelocity, VehicleForwardAxis);
+		ForwardsAcceleration = LocalAcceleration.X;
+	}
+}
+
+
+/**
+ * UChaosVehicleSimulation
+ */
+void UChaosVehicleSimulation::TickVehicle(UWorld* WorldIn, float DeltaTime, const FChaosVehicleDefaultAsyncInput& InputData, FChaosVehicleAsyncOutput& OutputData, Chaos::TPBDRigidParticleHandle<float, 3>* Handle)
+{
+	World = WorldIn;
+	RigidHandle = Handle;
+
+	// movement updates and replication
+	if (World && RigidHandle)
+	{
+		if (!VehicleState.bSleeping)
+		{
+			if (CanSimulate() && Handle)
+			{
+				UpdateSimulation(DeltaTime, InputData, Handle);
+				FillOutputState(OutputData);
+			}
+		}
+	}
+
+
+#if !(UE_BUILD_SHIPPING || UE_BUILD_TEST)
+	DrawDebug3D();
+#endif
+}
+
+void UChaosVehicleSimulation::UpdateSimulation(float DeltaTime, const FChaosVehicleDefaultAsyncInput& InputData, Chaos::TPBDRigidParticleHandle<float, 3>* Handle)
+{
+	VehicleState.CaptureState(Handle, InputData.GravityZ, DeltaTime);
+
+	ApplyAerodynamics(DeltaTime);
+	ApplyAerofoilForces(DeltaTime);
+	ApplyThrustForces(DeltaTime);
+	ApplyTorqueControl(DeltaTime, InputData);
+}
+
+void UChaosVehicleSimulation::FillOutputState(FChaosVehicleAsyncOutput& Output)
+{
+
+}
+
+/** Pass control Input to the vehicle systems */
+void UChaosVehicleSimulation::ApplyInput(const FControlInputs& ControlInputs, float DeltaTime)
+{
+	for (int AerofoilIdx = 0; AerofoilIdx < PVehicle->Aerofoils.Num(); AerofoilIdx++)
+	{
+		Chaos::FAerofoil& Aerofoil = PVehicle->GetAerofoil(AerofoilIdx);
+		switch (Aerofoil.Setup().Type)
+		{
+		case Chaos::EAerofoilType::Rudder:
+			Aerofoil.SetControlSurface(-ControlInputs.YawInput);
+			break;
+
+		case Chaos::EAerofoilType::Elevator:
+			Aerofoil.SetControlSurface(ControlInputs.PitchInput);
+			break;
+
+		case Chaos::EAerofoilType::Wing:
+			if (Aerofoil.Setup().Offset.Y < 0.0f)
+			{
+				Aerofoil.SetControlSurface(ControlInputs.RollInput);
+			}
+			else
+			{
+				Aerofoil.SetControlSurface(-ControlInputs.RollInput);
+			}
+			break;
+		}
+	}
+
+	for (int Thrusterdx = 0; Thrusterdx < PVehicle->Thrusters.Num(); Thrusterdx++)
+	{
+		Chaos::FSimpleThrustSim& Thruster = PVehicle->GetThruster(Thrusterdx);
+
+		Thruster.SetThrottle(ControlInputs.ThrottleInput);
+
+		switch (Thruster.Setup().Type)
+		{
+		case Chaos::EThrustType::HelicopterRotor:
+		{
+			Thruster.SetPitch(ControlInputs.PitchInput);
+			Thruster.SetRoll(ControlInputs.RollInput);
+		}
+		break;
+
+		case Chaos::EThrustType::Rudder:
+		{
+			Thruster.SetYaw(-ControlInputs.YawInput - ControlInputs.SteeringInput);
+		}
+		break;
+
+		case Chaos::EThrustType::Elevator:
+		{
+			Thruster.SetPitch(ControlInputs.PitchInput);
+		}
+		break;
+
+		case Chaos::EThrustType::Wing:
+		{
+			if (Thruster.Setup().Offset.Y < 0.0f)
+			{
+				Thruster.SetRoll(ControlInputs.RollInput);
+			}
+			else
+			{
+				Thruster.SetRoll(-ControlInputs.RollInput);
+			}
+		}
+		break;
+
+		}
+
+	}
+
+}
+
+
+void UChaosVehicleSimulation::ApplyAerodynamics(float DeltaTime)
+{
+	if (!GVehicleDebugParams.DisableAerodynamics)
+	{
+		// This force applied all the time whether the vehicle is on the ground or not
+		Chaos::FSimpleAerodynamicsSim& PAerodynamics = PVehicle->GetAerodynamics();
+		FVector LocalDragLiftForce = (PAerodynamics.GetCombinedForces(Chaos::CmToM(VehicleState.ForwardSpeed))) * Chaos::MToCmScaling();
+		FVector WorldLiftDragForce = VehicleState.VehicleWorldTransform.TransformVector(LocalDragLiftForce);
+		AddForce(WorldLiftDragForce);
+	}
+}
+
+FVector GetWorldVelocityAtPoint(const Chaos::TPBDRigidParticleHandle<float, 3>* RigidHandle, const FVector& WorldLocation)
+{
+	const Chaos::FVec3 COM = RigidHandle ? Chaos::FParticleUtilitiesGT::GetCoMWorldPosition(RigidHandle) : Chaos::FParticleUtilitiesGT::GetActorWorldTransform(RigidHandle).GetTranslation();
+	const Chaos::FVec3 Diff = WorldLocation - COM;
+	return RigidHandle->V() - Chaos::FVec3::CrossProduct(Diff, RigidHandle->W());
+}
+
+void UChaosVehicleSimulation::ApplyAerofoilForces(float DeltaTime)
+{
+	if (GVehicleDebugParams.DisableAerofoils || RigidHandle == nullptr)
+		return;
+
+	TArray<FVector> VelocityLocal;
+	TArray<FVector> VelocityWorld;
+	VelocityLocal.SetNum(PVehicle->Aerofoils.Num());
+	VelocityWorld.SetNum(PVehicle->Aerofoils.Num());
+
+	float Altitude = VehicleState.VehicleWorldTransform.GetLocation().Z;
+
+	// Work out velocity at each aerofoil before applying any forces so there's no bias on the first ones processed
+	for (int AerofoilIdx = 0; AerofoilIdx < PVehicle->Aerofoils.Num(); AerofoilIdx++)
+	{
+		FVector WorldLocation = VehicleState.VehicleWorldTransform.TransformPosition(PVehicle->GetAerofoil(AerofoilIdx).Setup().Offset * Chaos::MToCmScaling());
+		VelocityWorld[AerofoilIdx] = GetWorldVelocityAtPoint(RigidHandle, WorldLocation);
+		VelocityLocal[AerofoilIdx] = VehicleState.VehicleWorldTransform.InverseTransformVector(VelocityWorld[AerofoilIdx]);
+	}
+
+	for (int AerofoilIdx = 0; AerofoilIdx < PVehicle->Aerofoils.Num(); AerofoilIdx++)
+	{
+		Chaos::FAerofoil& Aerofoil = PVehicle->GetAerofoil(AerofoilIdx);
+
+		FVector LocalForce = Aerofoil.GetForce(VehicleState.VehicleWorldTransform, VelocityLocal[AerofoilIdx] * Chaos::CmToMScaling(), Chaos::CmToM(Altitude), DeltaTime);
+
+		FVector WorldForce = VehicleState.VehicleWorldTransform.TransformVector(LocalForce);
+		FVector WorldLocation = VehicleState.VehicleWorldTransform.TransformPosition(Aerofoil.GetCenterOfLiftOffset() * Chaos::MToCmScaling());
+		AddForceAtPosition(WorldForce * Chaos::MToCmScaling(), WorldLocation);
+
+#if !(UE_BUILD_SHIPPING || UE_BUILD_TEST)
+		FVector WorldAxis = VehicleState.VehicleWorldTransform.TransformVector(FVector::CrossProduct(FVector(1, 0, 0), Aerofoil.Setup().UpAxis));
+		if (GVehicleDebugParams.ShowAerofoilSurface)
+		{
+			DrawDebugLine(World, WorldLocation - WorldAxis * 150.0f, WorldLocation + WorldAxis * 150.0f, FColor::Black, false, -1.f, 0, 5.f);
+		}
+		if (GVehicleDebugParams.ShowAerofoilForces)
+		{
+			DrawDebugLine(World, WorldLocation, WorldLocation + WorldForce * GVehicleDebugParams.ForceDebugScaling, FColor::Green, false, -1.f, 0, 16.f);
+		}
+#endif
+	}
+
+}
+
+
+void UChaosVehicleSimulation::ApplyThrustForces(float DeltaTime)
+{
+	if (GVehicleDebugParams.DisableThrusters || RigidHandle == nullptr)
+		return;
+
+	for (int ThrusterIdx = 0; ThrusterIdx < PVehicle->Thrusters.Num(); ThrusterIdx++)
+	{
+		Chaos::FSimpleThrustSim& Thruster = PVehicle->GetThruster(ThrusterIdx);
+
+		FVector COM_Offset = RigidHandle->CenterOfMass();
+		COM_Offset.Z = 0.0f;
+		Thruster.SetWorldVelocity(VehicleState.VehicleWorldVelocity);
+
+		Thruster.Simulate(DeltaTime);
+		FVector ThrustWorldLocation = VehicleState.VehicleWorldTransform.TransformPosition(Thruster.GetThrustLocation() + COM_Offset);
+		FVector ThrustForce = VehicleState.VehicleWorldTransform.TransformPosition(Thruster.GetThrustForce());
+
+		AddForceAtPosition(ThrustForce, ThrustWorldLocation);
+	}
+
+}
+
+
+void UChaosVehicleSimulation::ApplyTorqueControl(float DeltaTime, const FChaosVehicleDefaultAsyncInput& InputData)
+{
+	if (!PVehicle->HasTorqueControlSetup())
+	{
+		return;
+	}
+
+	const FControlInputs& ControlInputsPT = InputData.ControlInputs;
+
+	if (!GVehicleDebugParams.DisableTorqueControl && RigidHandle)
+	{
+		FVector TotalTorque = FVector::ZeroVector;
+		if (PVehicle->HasTorqueControlSetup() && PVehicle->GetTorqueControl().Setup().Enabled)
+		{
+			const Chaos::FTargetRotationControlConfig& TargetRotationControl = PVehicle->GetTargetRotationControl().Setup();
+			auto ComputeTorque = [](const FVector& TargetUp, const FVector& CurrentUp, const FVector& AngVelocityWorld, float Stiffness, float Damping, float MaxAccel) -> FVector
+			{
+				const FQuat CurUpToTargetUp = FQuat::FindBetweenNormals(CurrentUp, TargetUp);
+				const FVector Axis = CurUpToTargetUp.GetRotationAxis();
+				const float Angle = CurUpToTargetUp.GetAngle();
+
+				float Strength = (Angle * Stiffness - FVector::DotProduct(AngVelocityWorld, Axis) * Damping);
+				Strength = FMath::Clamp(Strength, -MaxAccel, MaxAccel);
+				const FVector Torque = Axis * Strength;
+				return Torque;
+			};
+
+
+			FVector TargetUp = FVector(0.f, 0.f, 1.f);
+			float RollMaxAngleRadians = Chaos::DegToRad(TargetRotationControl.RollMaxAngle);
+			float PitchMaxAngleRadians = Chaos::DegToRad(TargetRotationControl.PitchMaxAngle);
+			float Speed = FMath::Min(Chaos::CmToM(VehicleState.ForwardSpeed), 20.0f); // cap here
+
+			float SpeeScaledRollAmount = 1.0f;
+			float TargetRoll = 0.f;
+			if (TargetRotationControl.bRollVsSpeedEnabled)
+			{
+				if (PVehicle->Wheels[0].InContact()) // HACK need IsAllowedToSteer virtual method
+				{
+					TargetRoll = ControlInputsPT.SteeringInput * TargetRotationControl.RollControlScaling * (Speed * Speed) * DeltaTime * 60.0f;
+				}
+			}
+			else
+			{
+				TargetRoll = ControlInputsPT.SteeringInput * TargetRotationControl.RollControlScaling;
+			}
+
+			FVector Rt = VehicleState.VehicleRightAxis * FMath::Max(FMath::Min(TargetRoll, RollMaxAngleRadians), -RollMaxAngleRadians);
+			FVector Pt = VehicleState.VehicleForwardAxis * FMath::Max(FMath::Min(ControlInputsPT.PitchInput * TargetRotationControl.PitchControlScaling, PitchMaxAngleRadians), -PitchMaxAngleRadians);
+
+			FVector UseUp = TargetUp + Rt + Pt;
+			UseUp.Normalize();
+
+			TargetUp = UseUp;
+
+			const FVector UpVector = VehicleState.VehicleUpAxis;
+			const FVector AngVelocityWorld = VehicleState.VehicleWorldAngularVelocity;
+
+			const FVector AirControlTorque = ComputeTorque(TargetUp, UpVector, AngVelocityWorld, TargetRotationControl.RotationStiffness, TargetRotationControl.RotationDamping, TargetRotationControl.MaxAccel);
+			const FVector ForwardVector = VehicleState.VehicleForwardAxis;
+			const FVector RightVector = VehicleState.VehicleRightAxis;
+
+			const float RollAirControl = FVector::DotProduct(AirControlTorque, ForwardVector);
+			const float PitchAirControl = FVector::DotProduct(AirControlTorque, RightVector);
+			const float YawAirControl = FVector::DotProduct(AirControlTorque, UpVector);
+
+			TotalTorque = RollAirControl * ForwardVector * TargetRotationControl.AutoCentreRollStrength
+				+ YawAirControl * UpVector * TargetRotationControl.AutoCentreYawStrength
+				+ PitchAirControl * RightVector * TargetRotationControl.AutoCentrePitchStrength;
+		}
+
+		if (PVehicle->HasTorqueControlSetup() && PVehicle->GetTorqueControl().Setup().Enabled)
+		{
+			const Chaos::FTorqueControlConfig& TorqueControl = PVehicle->GetTorqueControl().Setup();
+
+			TotalTorque -= VehicleState.VehicleForwardAxis * ControlInputsPT.RollInput * TorqueControl.RollTorqueScaling;
+			TotalTorque += VehicleState.VehicleRightAxis * ControlInputsPT.PitchInput * TorqueControl.PitchTorqueScaling;
+			TotalTorque += VehicleState.VehicleUpAxis * ControlInputsPT.YawInput * TorqueControl.YawTorqueScaling;
+			TotalTorque += VehicleState.VehicleUpAxis * ControlInputsPT.RollInput * TorqueControl.YawFromRollTorqueScaling;
+
+			// slowing rotation effect
+			FVector DampingTorque = (VehicleState.VehicleWorldAngularVelocity) * TorqueControl.RotationDamping;
+
+			// combined world torque
+			TotalTorque -= DampingTorque;
+		}
+
+		AddTorqueInRadians(TotalTorque, true, true);
+	}
+
+
+	if (!GVehicleDebugParams.DisableStabilizeControl && PVehicle->HasStabilizeControlSetup() && PVehicle->GetStabilizeControl().Setup().Enabled && RigidHandle)
+	{
+		const Chaos::FStabilizeControlConfig& StabilizeControl = PVehicle->GetStabilizeControl().Setup();
+
+		// try to cancel out velocity on Z axis
+		FVector CorrectionalForce = FVector::ZeroVector;
+		{
+			bool MaintainAltitude = true;
+			if (MaintainAltitude)
+			{
+				CorrectionalForce.Z = -StabilizeControl.AltitudeHoldZ * VehicleState.VehicleWorldVelocity.Z / DeltaTime;
+			}
+		}
+
+		// try to cancel out velocity on X/Y plane
+		// #todo: Will break helicopter setup??if (FMath::Abs(RollInput) < SMALL_NUMBER && FMath::Abs(PitchInput) < SMALL_NUMBER)
+		{
+			CorrectionalForce.X = -StabilizeControl.PositionHoldXY * VehicleState.VehicleWorldVelocity.X / DeltaTime;
+			CorrectionalForce.Y = -StabilizeControl.PositionHoldXY * VehicleState.VehicleWorldVelocity.Y / DeltaTime;
+		}
+		AddForce(CorrectionalForce);
+	}
+	
+}
+
+void UChaosVehicleSimulation::DrawDebug3D()
+{
+#if !(UE_BUILD_SHIPPING || UE_BUILD_TEST)
+	if (RigidHandle == nullptr)
+	{
+		return;
+	}
+
+	const FTransform BodyTransform = VehicleState.VehicleWorldTransform;
+
+	if (GVehicleDebugParams.ShowCOM)
+	{
+		const Chaos::FVec3 COMWorld = Chaos::FParticleUtilitiesGT::GetCoMWorldPosition(RigidHandle);
+		DrawDebugCoordinateSystem(World, COMWorld, FRotator(BodyTransform.GetRotation()), 200.f, false, -1.f, 0, 2.f);
+	}
+
+	if (GVehicleDebugParams.ShowModelOrigin)
+	{
+		DrawDebugCoordinateSystem(World, BodyTransform.GetLocation(), FRotator(BodyTransform.GetRotation()), 200.f, false, -1.f, 0, 2.f);
+	}
+#endif
+}
+
+void UChaosVehicleSimulation::AddForce(const FVector& Force, bool bAllowSubstepping /*= true*/, bool bAccelChange /*= false*/)
+{
+	if (ensure(RigidHandle))
+	{
+		Chaos::EObjectStateType ObjectState = RigidHandle->ObjectState();
+		if (CHAOS_ENSURE(ObjectState == Chaos::EObjectStateType::Dynamic || ObjectState == Chaos::EObjectStateType::Sleeping))
+		{
+			RigidHandle->SetObjectState(Chaos::EObjectStateType::Dynamic);
+
+			if (bAccelChange)
+			{
+				const float RigidMass = RigidHandle->M();
+				const Chaos::TVector<float, 3> Acceleration = Force * RigidMass;
+				RigidHandle->SetF(RigidHandle->F() + Acceleration);
+			}
+			else
+			{
+				RigidHandle->SetF(RigidHandle->F() + Force);
+			}
+
+		}
+
+#if !(UE_BUILD_SHIPPING || UE_BUILD_TEST)
+		if (GVehicleDebugParams.ShowAllForces)
+		{
+			FVector Position = VehicleState.VehicleWorldCOM;
+			DrawDebugDirectionalArrow(World, Position, Position + Force * GVehicleDebugParams.ForceDebugScaling
+				, 20.f, FColor::Blue, false, 0, 0, 2.f);
+		}
+#endif	
+	}
+
+}
+
+void UChaosVehicleSimulation::AddForceAtPosition(const FVector& Force, const FVector& Position, bool bAllowSubstepping /*= true*/, bool bIsLocalForce /*= false*/)
+{
+	if (ensure(RigidHandle))
+	{
+		const Chaos::FVec3 WorldCOM = Chaos::FParticleUtilitiesGT::GetCoMWorldPosition(RigidHandle);
+		RigidHandle->SetObjectState(Chaos::EObjectStateType::Dynamic);
+		const Chaos::FVec3 WorldTorque = Chaos::FVec3::CrossProduct(Position - WorldCOM, Force);
+		RigidHandle->SetF(Force + RigidHandle->F());
+		RigidHandle->SetTorque(WorldTorque + RigidHandle->Torque());
+
+#if !(UE_BUILD_SHIPPING || UE_BUILD_TEST)
+		if (GVehicleDebugParams.ShowAllForces)
+		{
+			DrawDebugDirectionalArrow(World, Position, Position + Force * GVehicleDebugParams.ForceDebugScaling
+				, 20.f, FColor::Blue, false, 0, 0, 2.f);
+		}
+#endif
+	}
+}
+
+void UChaosVehicleSimulation::AddImpulse(const FVector& Impulse, bool bVelChange)
+{
+	if (ensure(RigidHandle))
+	{
+		if (bVelChange)
+		{
+			RigidHandle->SetLinearImpulse(RigidHandle->LinearImpulse() + RigidHandle->M() * Impulse);
+		}
+		else
+		{
+			RigidHandle->SetLinearImpulse(RigidHandle->LinearImpulse() + Impulse);
+		}
+
+#if !(UE_BUILD_SHIPPING || UE_BUILD_TEST)
+		if (GVehicleDebugParams.ShowAllForces)
+		{
+			FVector Position = VehicleState.VehicleWorldCOM;
+			DrawDebugDirectionalArrow(World, Position, Position + Impulse * GVehicleDebugParams.ForceDebugScaling
+				, 20.f, FColor::Red, false, 0, 0, 2.f);
+		}
+#endif
+	}
+}
+
+void UChaosVehicleSimulation::AddImpulseAtPosition(const FVector& Impulse, const FVector& Position)
+{
+	if (ensure(RigidHandle))
+	{
+		const Chaos::FVec3 WorldCOM = Chaos::FParticleUtilitiesGT::GetCoMWorldPosition(RigidHandle);
+		const Chaos::FVec3 AngularImpulse = Chaos::FVec3::CrossProduct(Position - WorldCOM, Impulse);
+		RigidHandle->SetLinearImpulse(RigidHandle->LinearImpulse() + Impulse);
+		RigidHandle->SetAngularImpulse(RigidHandle->AngularImpulse() + AngularImpulse);
+
+
+#if !(UE_BUILD_SHIPPING || UE_BUILD_TEST)
+		if (GVehicleDebugParams.ShowAllForces)
+		{
+			DrawDebugDirectionalArrow(World, Position, Position + Impulse * GVehicleDebugParams.ForceDebugScaling
+				, 20.f, FColor::Red, false, 0, 0, 2.f);
+		}
+#endif
+	}
+
+}
+
+void UChaosVehicleSimulation::AddTorqueInRadians(const FVector& Torque, bool bAllowSubstepping /*= true*/, bool bAccelChange /*= false*/)
+{
+	if (ensure(RigidHandle))
+	{
+		Chaos::EObjectStateType ObjectState = RigidHandle->ObjectState();
+		if (CHAOS_ENSURE(ObjectState == Chaos::EObjectStateType::Dynamic || ObjectState == Chaos::EObjectStateType::Sleeping))
+		{
+			if (bAccelChange)
+			{
+				RigidHandle->SetTorque(RigidHandle->Torque() + Chaos::FParticleUtilitiesXR::GetWorldInertia(RigidHandle) * Torque);
+			}
+			else
+			{
+				RigidHandle->SetTorque(RigidHandle->Torque() + Torque);
+			}
+		}
+
+#if !(UE_BUILD_SHIPPING || UE_BUILD_TEST)
+		//#todo: how do we visualize torque?
+#endif
+
+	}
+}
+
+
+/**
+ * UChaosVehicleMovementComponent
+ */
 UChaosVehicleMovementComponent::UChaosVehicleMovementComponent(const FObjectInitializer& ObjectInitializer)
 	: Super(ObjectInitializer)
 {
@@ -119,6 +623,9 @@
 	RollInputRate.FallRate = 10.0f;
 	YawInputRate.RiseRate = 6.0f;
 	YawInputRate.FallRate = 10.0f;
+	TargetGearInput = 0;
+	SetGearImmediate = -1;
+	TransmissionType = Chaos::ETransmissionType::Automatic;
 
 	SetIsReplicatedByDefault(true);
 
@@ -189,7 +696,7 @@
 
 bool UChaosVehicleMovementComponent::HasValidPhysicsState() const
 {
-	return PVehicle.IsValid();
+	return PVehicleOutput.IsValid();
 }
 
 bool UChaosVehicleMovementComponent::CanCreateVehicle() const
@@ -230,7 +737,7 @@
 			CreateVehicle();
 			FixupSkeletalMesh();
 
-			if (PVehicle)
+			if (PVehicleOutput)
 			{
 				FChaosVehicleManager* VehicleManager = FChaosVehicleManager::GetVehicleManagerFromScene(PhysScene);
 				VehicleManager->AddVehicle(this);
@@ -250,11 +757,11 @@
 {
 	Super::OnDestroyPhysicsState();
 
-	if (PVehicle.IsValid())
+	if (PVehicleOutput.IsValid())
 	{
 		FChaosVehicleManager* VehicleManager = FChaosVehicleManager::GetVehicleManagerFromScene(GetWorld()->GetPhysicsScene());
 		VehicleManager->RemoveVehicle(this);
-		PVehicle.Reset(nullptr);
+		PVehicleOutput.Reset(nullptr);
 
 		if (UpdatedComponent)
 		{
@@ -263,10 +770,10 @@
 	}
 }
 
-void UChaosVehicleMovementComponent::PreTick(float DeltaTime)
+void UChaosVehicleMovementComponent::PreTickGT(float DeltaTime)
 {
 	// movement updates and replication
-	if (PVehicle && UpdatedComponent)
+	if (PVehicleOutput && UpdatedComponent)
 	{
 		APawn* MyOwner = Cast<APawn>(UpdatedComponent->GetOwner());
 		if (MyOwner)
@@ -275,38 +782,27 @@
 		}
 	}
 
+	if (!GVehicleDebugParams.DisableVehicleSleep)
+	{
+		// is this needless copying
+		FControlInputs ControlInputs;
+		ControlInputs.ThrottleInput = ThrottleInput;
+		ControlInputs.BrakeInput = BrakeInput;
+		ControlInputs.SteeringInput = SteeringInput;
+		ControlInputs.HandbrakeInput = HandbrakeInput;
+		ControlInputs.RollInput = RollInput;
+		ControlInputs.PitchInput = PitchInput;
+		ControlInputs.YawInput = YawInput;
+		ProcessSleeping(ControlInputs);
+	}
+
+
 	if (VehicleSetupTag != FChaosVehicleManager::VehicleSetupTag)
 	{
 		RecreatePhysicsState();
 	}
 }
 
-void UChaosVehicleMovementComponent::TickVehicle(float DeltaTime)
-{
-	// movement updates and replication
-	FBodyInstance* TargetInstance = GetBodyInstance();
-	if (PVehicle && UpdatedComponent && TargetInstance)
-	{
-		APawn* MyOwner = Cast<APawn>(UpdatedComponent->GetOwner());
-		if (MyOwner)
-		{
-			if (!GVehicleDebugParams.DisableVehicleSleep)
-			{
-				ProcessSleeping();
-			}
-
-			if (!VehicleState.bSleeping)
-			{
-				UpdateSimulation(DeltaTime);
-			}
-		}
-	}
-
-#if !(UE_BUILD_SHIPPING || UE_BUILD_TEST)
-	DrawDebug3D();
-#endif
-}
-
 void UChaosVehicleMovementComponent::StopMovementImmediately()
 {
 	FBodyInstance* TargetInstance = GetBodyInstance();
@@ -380,57 +876,33 @@
 
 void UChaosVehicleMovementComponent::SetTargetGear(int32 GearNum, bool bImmediate)
 {
-	if (PVehicle.IsValid() && PVehicle->HasTransmission() && GearNum != PVehicle->GetTransmission().GetTargetGear())
+	if (GearNum != TargetGearInput)
 	{
-		PVehicle->GetTransmission().SetGear(GearNum, bImmediate);
+		TargetGearInput = GearNum;
+		SetGearImmediate = bImmediate?1:0; // true == 1, false == 0, don't transmit to physics thread == -1
 	}
 }
 
 void UChaosVehicleMovementComponent::SetUseAutomaticGears(bool bUseAuto)
 {
-	if (PVehicle.IsValid() && PVehicle->HasTransmission())
-	{
-		Chaos::ETransmissionType TransmissionType = bUseAuto ? Chaos::ETransmissionType::Automatic : Chaos::ETransmissionType::Manual;
-		PVehicle->GetTransmission().AccessSetup().TransmissionType = TransmissionType;
-	}
+	TransmissionType = bUseAuto ? Chaos::ETransmissionType::Automatic : Chaos::ETransmissionType::Manual;
 }
 
 // Data access
 
 int32 UChaosVehicleMovementComponent::GetCurrentGear() const
 {
-	int32 CurrentGear = 0;
-
-	if (PVehicle.IsValid() && PVehicle->HasTransmission())
-	{
-		CurrentGear = PVehicle->GetTransmission().GetCurrentGear();
-	}
-
-	return CurrentGear;
+	return PVehicleOutput->CurrentGear;
 }
 
 int32 UChaosVehicleMovementComponent::GetTargetGear() const
 {
-	int32 TargetGear = 0;
-
-	if (PVehicle.IsValid() && PVehicle->HasTransmission())
-	{
-		TargetGear = PVehicle->GetTransmission().GetTargetGear();
-	}
-
-	return TargetGear;
+	return PVehicleOutput->TargetGear; //  or TargetGearInput
 }
 
 bool UChaosVehicleMovementComponent::GetUseAutoGears() const
 {
-	bool UseAutoGears = 0;
-
-	if (PVehicle.IsValid() && PVehicle->HasTransmission())
-	{
-		UseAutoGears = PVehicle->GetTransmission().Setup().TransmissionType == Chaos::ETransmissionType::Automatic;
-	}
-
-	return UseAutoGears;
+	return (TransmissionType == Chaos::ETransmissionType::Automatic);
 }
 
 float UChaosVehicleMovementComponent::GetForwardSpeed() const
@@ -501,7 +973,6 @@
 			if (VehicleState.ForwardSpeed < StopThreshold && VehicleState.ForwardSpeed > -StopThreshold)	//auto brake 
 			{
 				NewBrakeInput = 1.f;
-				FBodyInstance* TargetInstance = GetBodyInstance();
 			}
 		}
 
@@ -533,15 +1004,15 @@
 float UChaosVehicleMovementComponent::CalcThrottleInput()
 {
 	float NewThrottleInput = RawThrottleInput;
-	if (bReverseAsBrake && PVehicle->HasTransmission())
+	if (bReverseAsBrake)
 	{
-		if (RawBrakeInput > 0.f && PVehicle->GetTransmission().GetTargetGear() < 0/*ForwardSpeed < -WrongDirectionThreshold*/)
+		if (RawBrakeInput > 0.f && GetTargetGear() < 0)
 		{
 			NewThrottleInput = RawBrakeInput;
 		}
 		else
 			//If the user is changing direction we should really be braking first and not applying any gas, so wait until they've changed gears
-			if ((RawThrottleInput > 0.f && PVehicle->GetTransmission().GetTargetGear() < 0) || (RawBrakeInput > 0.f && PVehicle->GetTransmission().GetTargetGear() > 0))
+			if ((RawThrottleInput > 0.f && GetTargetGear() < 0) || (RawBrakeInput > 0.f && GetTargetGear() > 0))
 			{
 				NewThrottleInput = 0.f;
 			}
@@ -562,13 +1033,13 @@
 
 	// Send this immediately.
 	int32 CurrentGear = 0;
-	if (PVehicle && PVehicle->HasTransmission())
+	if (PVehicleOutput)
 	{
-		CurrentGear = PVehicle->GetTransmission().GetCurrentGear();
+		CurrentGear = PVehicleOutput->CurrentGear;
 	}
 
 	AController* Controller = GetController();
-	if (Controller && Controller->IsLocalController() && PVehicle)
+	if (Controller && Controller->IsLocalController() && PVehicleOutput)
 	{
 		ServerUpdateState(SteeringInput, ThrottleInput, BrakeInput, HandbrakeInput, CurrentGear, RollInput, PitchInput, YawInput);
 	}
@@ -593,39 +1064,37 @@
 {
 	// update input values
 	AController* Controller = GetController();
+	VehicleState.CaptureState(GetBodyInstance(), GetGravityZ(), DeltaTime);
 
 	// IsLocallyControlled will fail if the owner is unpossessed (i.e. Controller == nullptr);
 	// Should we remove input instead of relying on replicated state in that case?
-	if (Controller && Controller->IsLocalController() && PVehicle)
+	if (Controller && Controller->IsLocalController() && PVehicleOutput)
 	{
-		if (PVehicle->HasTransmission())
+		if (bReverseAsBrake)
 		{
-			if (bReverseAsBrake)
+			//for reverse as state we want to automatically shift between reverse and first gear
+			if (FMath::Abs(GetForwardSpeed()) < WrongDirectionThreshold)	//we only shift between reverse and first if the car is slow enough.
 			{
-				//for reverse as state we want to automatically shift between reverse and first gear
-				if (FMath::Abs(GetForwardSpeed()) < WrongDirectionThreshold)	//we only shift between reverse and first if the car is slow enough.
+				if (RawBrakeInput > KINDA_SMALL_NUMBER && GetCurrentGear() >= 0 && GetTargetGear() >= 0)
 				{
-					if (RawBrakeInput > KINDA_SMALL_NUMBER && PVehicle->GetTransmission().GetCurrentGear() >= 0 && PVehicle->GetTransmission().GetTargetGear() >= 0)
-					{
-						SetTargetGear(-1, false);
-					}
-					else if (RawThrottleInput > KINDA_SMALL_NUMBER && PVehicle->GetTransmission().GetCurrentGear() <= 0 && PVehicle->GetTransmission().GetTargetGear() <= 0)
-					{
-						SetTargetGear(1, false);
-					}
+					SetTargetGear(-1, false);
 				}
-			}
-			else
-			{
-				if (PVehicle->GetTransmission().Setup().TransmissionType == Chaos::ETransmissionType::Automatic
-					&& RawThrottleInput > KINDA_SMALL_NUMBER
-					&& PVehicle->GetTransmission().GetCurrentGear() == 0
-					&& PVehicle->GetTransmission().GetTargetGear() == 0)
+				else if (RawThrottleInput > KINDA_SMALL_NUMBER && GetCurrentGear() <= 0 && GetTargetGear() <= 0)
 				{
-					SetTargetGear(1, true);
+					SetTargetGear(1, false);
 				}
 			}
 		}
+		else
+		{
+			if (TransmissionType == Chaos::ETransmissionType::Automatic
+				&& RawThrottleInput > KINDA_SMALL_NUMBER
+				&& PVehicleOutput->CurrentGear == 0
+				&& PVehicleOutput->TargetGear == 0)
+			{
+				SetTargetGear(1, true);
+			}
+		}
 
 		SteeringInput = SteeringInputRate.InterpInputValue(DeltaTime, SteeringInput, CalcSteeringInput());
 		ThrottleInput = ThrottleInputRate.InterpInputValue(DeltaTime, ThrottleInput, CalcThrottleInput());
@@ -636,12 +1105,7 @@
 		HandbrakeInput = HandbrakeInputRate.InterpInputValue(DeltaTime, HandbrakeInput, CalcHandbrakeInput());
 
 		// and send to server - (ServerUpdateState_Implementation below)
-		int32 TargetGear = 0;
-		if (PVehicle->HasTransmission())
-		{
-			TargetGear = PVehicle->GetTransmission().GetTargetGear();
-		}
-		ServerUpdateState(SteeringInput, ThrottleInput, BrakeInput, HandbrakeInput, TargetGear, RollInput, PitchInput, YawInput);
+		ServerUpdateState(SteeringInput, ThrottleInput, BrakeInput, HandbrakeInput, TargetGearInput, RollInput, PitchInput, YawInput);
 
 		if (PawnOwner && PawnOwner->IsNetMode(NM_Client))
 		{
@@ -662,284 +1126,8 @@
 	}
 }
 
-void UChaosVehicleMovementComponent::UpdateSimulation(float DeltaTime)
-{
-	FBodyInstance* TargetInstance = GetBodyInstance();
 
-	if (CanSimulate() && TargetInstance)
-	{
-		VehicleState.CaptureState(TargetInstance, GetGravityZ(), DeltaTime);
-
-		ApplyAerodynamics(DeltaTime);
-		ApplyAerofoilForces(DeltaTime);
-		ApplyThrustForces(DeltaTime);
-		ApplyTorqueControl(DeltaTime);
-	}
-}
-
-/** Pass control Input to the vehicle systems */
-void UChaosVehicleMovementComponent::ApplyInput(float DeltaTime)
-{
-	for (int AerofoilIdx = 0; AerofoilIdx < Aerofoils.Num(); AerofoilIdx++)
-	{
-		Chaos::FAerofoil& Aerofoil = PVehicle->GetAerofoil(AerofoilIdx);
-		switch (Aerofoil.Setup().Type)
-		{
-			case Chaos::EAerofoilType::Rudder:
-			Aerofoil.SetControlSurface(-YawInput);
-			break;
-
-			case Chaos::EAerofoilType::Elevator:
-			Aerofoil.SetControlSurface(PitchInput);
-			break;
-
-			case Chaos::EAerofoilType::Wing:
-			if (Aerofoil.Setup().Offset.Y < 0.0f)
-			{
-				Aerofoil.SetControlSurface(RollInput);
-			}
-			else
-			{
-				Aerofoil.SetControlSurface(-RollInput);
-			}
-			break;
-		}
-	}
-
-	for (int Thrusterdx = 0; Thrusterdx < Thrusters.Num(); Thrusterdx++)
-	{
-		Chaos::FSimpleThrustSim& Thruster = PVehicle->GetThruster(Thrusterdx);
-
-		Thruster.SetThrottle(ThrottleInput);
-
-		switch (Thruster.Setup().Type)
-		{
-			case Chaos::EThrustType::HelicopterRotor:
-			{
-				Thruster.SetPitch(PitchInput);
-				Thruster.SetRoll(RollInput);
-			}
-			break;
-
-			case Chaos::EThrustType::Rudder:
-			{
-				Thruster.SetYaw(-YawInput - SteeringInput);
-			}
-			break;
-
-			case Chaos::EThrustType::Elevator:
-			{
-				Thruster.SetPitch(PitchInput);
-			}
-			break;
-
-			case Chaos::EThrustType::Wing:
-			{
-				if (Thruster.Setup().Offset.Y < 0.0f)
-				{
-					Thruster.SetRoll(RollInput);
-				}
-				else
-				{
-					Thruster.SetRoll(-RollInput);
-				}
-			}
-			break;
-
-
-		}
-
-	}
-
-}
-
-
-void UChaosVehicleMovementComponent::ApplyAerodynamics(float DeltaTime)
-{
-	if (!GVehicleDebugParams.DisableAerodynamics)
-	{
-		// This force applied all the time whether the vehicle is on the ground or not
-		Chaos::FSimpleAerodynamicsSim& PAerodynamics = PVehicle->GetAerodynamics();
-		FVector LocalDragLiftForce = (PAerodynamics.GetCombinedForces(Chaos::CmToM(VehicleState.ForwardSpeed))) * Chaos::MToCmScaling();
-		FVector WorldLiftDragForce = VehicleState.VehicleWorldTransform.TransformVector(LocalDragLiftForce);
-		AddForce(WorldLiftDragForce);
-	}
-}
-
-void UChaosVehicleMovementComponent::ApplyAerofoilForces(float DeltaTime)
-{
-	if (GVehicleDebugParams.DisableAerofoils || GetBodyInstance() == nullptr)
-		return;
-
-	TArray<FVector> VelocityLocal;
-	TArray<FVector> VelocityWorld;
-	VelocityLocal.SetNum(PVehicle->Aerofoils.Num());
-	VelocityWorld.SetNum(PVehicle->Aerofoils.Num());
-
-	float Altitude = VehicleState.VehicleWorldTransform.GetLocation().Z;
-
-	// Work out velocity at each aerofoil before applying any forces so there's no bias on the first ones processed
-	for (int AerofoilIdx = 0; AerofoilIdx < PVehicle->Aerofoils.Num(); AerofoilIdx++)
-	{
-		FVector WorldLocation = VehicleState.VehicleWorldTransform.TransformPosition(PVehicle->GetAerofoil(AerofoilIdx).Setup().Offset * Chaos::MToCmScaling());
-		VelocityWorld[AerofoilIdx] = GetBodyInstance()->GetUnrealWorldVelocityAtPoint(WorldLocation);
-		VelocityLocal[AerofoilIdx] = VehicleState.VehicleWorldTransform.InverseTransformVector(VelocityWorld[AerofoilIdx]);
-	}
-
-	for (int AerofoilIdx = 0; AerofoilIdx < PVehicle->Aerofoils.Num(); AerofoilIdx++)
-	{
-		Chaos::FAerofoil& Aerofoil = PVehicle->GetAerofoil(AerofoilIdx);
-
-		FVector LocalForce = Aerofoil.GetForce(VehicleState.VehicleWorldTransform, VelocityLocal[AerofoilIdx] * Chaos::CmToMScaling(), Chaos::CmToM(Altitude), DeltaTime);
-
-		FVector WorldForce = VehicleState.VehicleWorldTransform.TransformVector(LocalForce);
-		FVector WorldLocation = VehicleState.VehicleWorldTransform.TransformPosition(Aerofoil.GetCenterOfLiftOffset() * Chaos::MToCmScaling());
-		AddForceAtPosition(WorldForce * Chaos::MToCmScaling(), WorldLocation);
-
-#if !(UE_BUILD_SHIPPING || UE_BUILD_TEST)
-		FVector WorldAxis = VehicleState.VehicleWorldTransform.TransformVector(FVector::CrossProduct(FVector(1,0,0), Aerofoil.Setup().UpAxis));
-		if (GVehicleDebugParams.ShowAerofoilSurface)
-		{
-			DrawDebugLine(GetWorld(), WorldLocation - WorldAxis * 150.0f, WorldLocation + WorldAxis * 150.0f, FColor::Black, false, -1.f, 0, 5.f);
-		}
-		if (GVehicleDebugParams.ShowAerofoilForces)
-		{
-			DrawDebugLine(GetWorld(), WorldLocation, WorldLocation + WorldForce * GVehicleDebugParams.ForceDebugScaling, FColor::Green, false, -1.f, 0, 16.f);
-		}
-#endif
-	}
-
-}
-
-
-void UChaosVehicleMovementComponent::ApplyThrustForces(float DeltaTime)
-{
-	if (GVehicleDebugParams.DisableThrusters || GetBodyInstance() == nullptr)
-		return;
-
-	for (int ThrusterIdx = 0; ThrusterIdx < PVehicle->Thrusters.Num(); ThrusterIdx++)
-	{
-		Chaos::FSimpleThrustSim& Thruster = PVehicle->GetThruster(ThrusterIdx);
-
-		FVector COM_Offset = GetBodyInstance()->GetMassSpaceLocal().GetLocation();
-		COM_Offset.Z = 0.0f;
-		Thruster.SetWorldVelocity(VehicleState.VehicleWorldVelocity);
-
-		Thruster.Simulate(DeltaTime);
-		FVector ThrustWorldLocation = VehicleState.VehicleWorldTransform.TransformPosition(Thruster.GetThrustLocation() + COM_Offset);
-		FVector ThrustForce = VehicleState.VehicleWorldTransform.TransformPosition(Thruster.GetThrustForce());
-
-		AddForceAtPosition(ThrustForce, ThrustWorldLocation);
-	}
-
-}
-
-
-void UChaosVehicleMovementComponent::ApplyTorqueControl(float DeltaTime)
-{
-	FBodyInstance* TargetInstance = GetBodyInstance();
-	if (!GVehicleDebugParams.DisableTorqueControl && TargetInstance)
-	{
-		FVector TotalTorque = FVector::ZeroVector;
-		if (TargetRotationControl.Enabled)
-		{
-			auto ComputeTorque = [](const FVector& TargetUp, const FVector& CurrentUp, const FVector& AngVelocityWorld, float Stiffness, float Damping, float MaxAccel) -> FVector
-			{
-				const FQuat CurUpToTargetUp = FQuat::FindBetweenNormals(CurrentUp, TargetUp);
-				const FVector Axis = CurUpToTargetUp.GetRotationAxis();
-				const float Angle = CurUpToTargetUp.GetAngle();
-
-				float Strength = (Angle * Stiffness - FVector::DotProduct(AngVelocityWorld, Axis) * Damping);
-				Strength = FMath::Clamp(Strength, -MaxAccel, MaxAccel);
-				const FVector Torque = Axis * Strength;
-				return Torque;
-			};
-
-			 
-			FVector TargetUp = FVector(0.f, 0.f, 1.f);
-			float RollMaxAngleRadians = Chaos::DegToRad(TargetRotationControl.RollMaxAngle);
-			float PitchMaxAngleRadians = Chaos::DegToRad(TargetRotationControl.PitchMaxAngle);
-			float Speed = FMath::Min(Chaos::CmToM(VehicleState.ForwardSpeed), 20.0f); // cap here
-
-			float SpeeScaledRollAmount = 1.0f;
-			float TargetRoll = 0.f;
-			if (TargetRotationControl.bRollVsSpeedEnabled)
-			{
-				if (PVehicle->Wheels[0].InContact()) // HACK need IsAllowedToSteer virtual method
-				{
-					TargetRoll = SteeringInput * TargetRotationControl.RollControlScaling * (Speed * Speed) * DeltaTime * 60.0f;
-				}
-			}
-			else
-			{
-				TargetRoll = SteeringInput * TargetRotationControl.RollControlScaling;
-			}
-
-			FVector Rt = VehicleState.VehicleRightAxis * FMath::Max(FMath::Min(TargetRoll, RollMaxAngleRadians), -RollMaxAngleRadians);
-			FVector Pt = VehicleState.VehicleForwardAxis * FMath::Max(FMath::Min(PitchInput * TargetRotationControl.PitchControlScaling, PitchMaxAngleRadians), -PitchMaxAngleRadians);
-
-			FVector UseUp = TargetUp + Rt + Pt;
-			UseUp.Normalize();
-
-			TargetUp = UseUp;
-
-			const FVector UpVector = VehicleState.VehicleUpAxis;
-			const FVector AngVelocityWorld = VehicleState.VehicleWorldAngularVelocity;
-
-			const FVector AirControlTorque = ComputeTorque(TargetUp, UpVector, AngVelocityWorld,TargetRotationControl.RotationStiffness, TargetRotationControl.RotationDamping, TargetRotationControl.MaxAccel);
-			const FVector ForwardVector = VehicleState.VehicleForwardAxis;
-			const FVector RightVector = VehicleState.VehicleRightAxis;
-
-			const float RollAirControl = FVector::DotProduct(AirControlTorque, ForwardVector);
-			const float PitchAirControl = FVector::DotProduct(AirControlTorque, RightVector);
-			const float YawAirControl = FVector::DotProduct(AirControlTorque, UpVector);
-
-			TotalTorque = RollAirControl * ForwardVector * TargetRotationControl.AutoCentreRollStrength
-				+ YawAirControl * UpVector * TargetRotationControl.AutoCentreYawStrength
-				+ PitchAirControl * RightVector * TargetRotationControl.AutoCentrePitchStrength;
-		}
-
-		if (TorqueControl.Enabled)
-		{
-			TotalTorque -= VehicleState.VehicleForwardAxis * RollInput * TorqueControl.RollTorqueScaling;
-			TotalTorque += VehicleState.VehicleRightAxis * PitchInput * TorqueControl.PitchTorqueScaling;
-			TotalTorque += VehicleState.VehicleUpAxis * YawInput * TorqueControl.YawTorqueScaling;
-			TotalTorque += VehicleState.VehicleUpAxis * RollInput * TorqueControl.YawFromRollTorqueScaling;
-
-			// slowing rotation effect
-			FVector DampingTorque = (VehicleState.VehicleWorldAngularVelocity/* / DeltaTime*/) * TorqueControl.RotationDamping;
-
-			// combined world torque
-			TotalTorque -= DampingTorque;
-		}
-
-		AddTorqueInRadians(TotalTorque /*/ DeltaTime * 0.016f*/, /*bAllowSubstepping=*/true, /*bAccelChange=*/true);
-	}
-
-
-	if (!GVehicleDebugParams.DisableStabilizeControl && StabilizeControl.Enabled && TargetInstance)
-	{
-		// try to cancel out velocity on Z axis
-		FVector CorrectionalForce = FVector::ZeroVector;
-		{
-			bool MaintainAltitude = true;
-			if (MaintainAltitude)
-			{
-				CorrectionalForce.Z = -StabilizeControl.AltitudeHoldZ * VehicleState.VehicleWorldVelocity.Z / DeltaTime;
-			}
-		}
-
-		// try to cancel out velocity on X/Y plane
-		// #todo: Will break helicopter setup??if (FMath::Abs(RollInput) < SMALL_NUMBER && FMath::Abs(PitchInput) < SMALL_NUMBER)
-		{
-			CorrectionalForce.X = -StabilizeControl.PositionHoldXY * VehicleState.VehicleWorldVelocity.X / DeltaTime;
-			CorrectionalForce.Y = -StabilizeControl.PositionHoldXY * VehicleState.VehicleWorldVelocity.Y / DeltaTime;
-		}
-		AddForce(CorrectionalForce);
-	}
-}
-
-void UChaosVehicleMovementComponent::ProcessSleeping()
+void UChaosVehicleMovementComponent::ProcessSleeping(const FControlInputs& ControlInputs)
 {
 	FBodyInstance* TargetInstance = GetBodyInstance();
 	if (TargetInstance)
@@ -960,12 +1148,12 @@
 		// NOTE: Even on local clients, ReplicatedState will still be populated (the call to ServerUpdateState will
 		//			be processed locally). Maybe we should *just* use ReplicatedState?
 
-		// TODO: What about other inputs, like handbrake, roll, pitch, yaw?
 		const AController* Controller = GetController();
 		const bool bIsLocallyControlled = (Controller && Controller->IsLocalController());
-		const bool bControlInputPressed = bIsLocallyControlled ?
-			(RawThrottleInput >= SMALL_NUMBER) || (RawBrakeInput >= SMALL_NUMBER) || (FMath::Abs(RawSteeringInput) > SMALL_NUMBER) :
-			(ReplicatedState.ThrottleInput >= SMALL_NUMBER) || (ReplicatedState.BrakeInput >= SMALL_NUMBER) || (FMath::Abs(ReplicatedState.SteeringInput) > SMALL_NUMBER);
+		const bool bControlInputPressed = bIsLocallyControlled ? (ControlInputs.ThrottleInput >= SMALL_NUMBER) || (ControlInputs.BrakeInput >= SMALL_NUMBER) || (ControlInputs.SteeringInput >= SMALL_NUMBER)
+			|| (ControlInputs.RollInput >= SMALL_NUMBER) || (ControlInputs.PitchInput >= SMALL_NUMBER) || (ControlInputs.YawInput >= SMALL_NUMBER)
+			: (ReplicatedState.ThrottleInput >= SMALL_NUMBER) || (ReplicatedState.BrakeInput >= SMALL_NUMBER) || (ReplicatedState.SteeringInput >= SMALL_NUMBER)
+			|| (ReplicatedState.RollInput >= SMALL_NUMBER) || (ReplicatedState.PitchInput >= SMALL_NUMBER) || (ReplicatedState.YawInput >= SMALL_NUMBER);
 
 		// Wake if control input pressed
 		if (VehicleState.bSleeping && (bControlInputPressed || !VehicleState.bAllWheelsOnGround))
@@ -1102,28 +1290,31 @@
 {
 	ComputeConstants();
 
-	if (PVehicle == nullptr)
 	{
 		if (CanCreateVehicle())
 		{
 			check(UpdatedComponent);
 			if (ensure(UpdatedPrimitive != nullptr))
 			{
+				TUniquePtr<Chaos::FSimpleWheeledVehicle> PVehicle = CreatePhysicsVehicle();
+
 				// Low level physics representation
-				CreatePhysicsVehicle();
+				SetupVehicle(PVehicle); 
 
-				SetupVehicle();
-
-				if (PVehicle != nullptr)
+				if (PVehicleOutput != nullptr)
 				{
 					PostSetupVehicle();
 				}
+
+				// Physics thread simulation class will now take ownership of the PVehicle pointer, we cannot safely use it anymore from the game thread
+				VehicleSimulationPT->Init(PVehicle);
+
 			}
 		}
 	}
 }
 
-void UChaosVehicleMovementComponent::SetupVehicle()
+void UChaosVehicleMovementComponent::SetupVehicle(TUniquePtr<Chaos::FSimpleWheeledVehicle>& PVehicle)
 {
 	Chaos::FSimpleAerodynamicsSim AerodynamicsSim(&GetAerodynamicsConfig());
 	PVehicle->Aerodynamics.Add(AerodynamicsSim);
@@ -1204,7 +1395,7 @@
 {
 #if !(UE_BUILD_SHIPPING || UE_BUILD_TEST)
 	FBodyInstance* TargetInstance = GetBodyInstance();
-	if (!PVehicle.IsValid() || TargetInstance == nullptr)
+	if (PVehicleOutput == nullptr || TargetInstance == nullptr)
 	{
 		return;
 	}
@@ -1225,7 +1416,7 @@
 			YPos += Canvas->DrawText(RenderFont, FString::Printf(TEXT("Inertia : %s"), *TargetInstance->GetBodyInertiaTensor().ToString()), 4, YPos);
 		}
 
-		YPos += Canvas->DrawText(RenderFont, FString::Printf(TEXT("Awake %d"), TargetInstance->IsInstanceAwake()), 4, YPos);
+		YPos += Canvas->DrawText(RenderFont, FString::Printf(TEXT("Awake %d (Vehicle Sleep %d)"), TargetInstance->IsInstanceAwake(), VehicleState.bSleeping), 4, YPos);
 		YPos += Canvas->DrawText(RenderFont, FString::Printf(TEXT("Speed (km/h): %.1f  (MPH): %.1f  (m/s): %.1f"), ForwardSpeedKmH, ForwardSpeedMPH, ForwardSpeedMSec), 4, YPos);
 		YPos += Canvas->DrawText(RenderFont, FString::Printf(TEXT("Acceleration (m/s-2): %.1f"), Chaos::CmToM(VehicleState.LocalAcceleration.X)), 4, YPos);
 		YPos += Canvas->DrawText(RenderFont, FString::Printf(TEXT("GForce : %2.1f"), VehicleState.LocalGForce.X), 4, YPos);
@@ -1242,30 +1433,6 @@
 #endif
 }
 
-void UChaosVehicleMovementComponent::DrawDebug3D()
-{
-#if !(UE_BUILD_SHIPPING || UE_BUILD_TEST)
-	FBodyInstance* TargetInstance = GetBodyInstance();
-	if (TargetInstance == nullptr)
-	{
-		return;
-	}
-		
-	const FTransform BodyTransform = VehicleState.VehicleWorldTransform;
-
-	if (GVehicleDebugParams.ShowCOM && TargetInstance)
-	{
-		FVector COMWorld = TargetInstance->GetCOMPosition();
-		DrawDebugCoordinateSystem(GetWorld(), COMWorld, FRotator(BodyTransform.GetRotation()), 200.f, false, -1.f, 0, 2.f);
-	}
-
-	if (GVehicleDebugParams.ShowModelOrigin && TargetInstance)
-	{
-		DrawDebugCoordinateSystem(GetWorld(), BodyTransform.GetLocation(), FRotator(BodyTransform.GetRotation()), 200.f, false, -1.f, 0, 2.f);
-	}
-#endif
-}
-
 /// @cond DOXYGEN_WARNINGS
 
 void UChaosVehicleMovementComponent::GetLifetimeReplicatedProps( TArray< FLifetimeProperty > & OutLifetimeProps ) const
@@ -1291,81 +1458,12 @@
 }
 
 
-void UChaosVehicleMovementComponent::AddForce(const FVector& Force, bool bAllowSubstepping /*= true*/, bool bAccelChange /*= false*/)
+TUniquePtr<Chaos::FSimpleWheeledVehicle> UChaosVehicleMovementComponent::CreatePhysicsVehicle()
 {
-	check(GetBodyInstance());
-	GetBodyInstance()->AddForce(Force, bAllowSubstepping, bAccelChange);
-
-#if !(UE_BUILD_SHIPPING || UE_BUILD_TEST)
-	if (GVehicleDebugParams.ShowAllForces)
-	{
-		FVector Position = VehicleState.VehicleWorldCOM;
-		DrawDebugDirectionalArrow(GetWorld(), Position, Position + Force * GVehicleDebugParams.ForceDebugScaling
-			, 20.f, FColor::Blue, false, 0, 0, 2.f);
-	}
-#endif
+	PVehicleOutput = MakeUnique<FPhysicsVehicleOutput>();	// create physics output container
+	return MakeUnique<Chaos::FSimpleWheeledVehicle>();		// create physics sim
 }
 
-void UChaosVehicleMovementComponent::AddForceAtPosition(const FVector& Force, const FVector& Position, bool bAllowSubstepping /*= true*/, bool bIsLocalForce /*= false*/)
-{
-	check(GetBodyInstance());
-	GetBodyInstance()->AddForceAtPosition(Force, Position, bAllowSubstepping, bIsLocalForce);
-
-#if !(UE_BUILD_SHIPPING || UE_BUILD_TEST)
-	if (GVehicleDebugParams.ShowAllForces)
-	{
-		DrawDebugDirectionalArrow(GetWorld(), Position, Position + Force * GVehicleDebugParams.ForceDebugScaling
-			, 20.f, FColor::Blue, false, 0, 0, 2.f);
-	}
-#endif
-}
-
-void UChaosVehicleMovementComponent::AddImpulse(const FVector& Impulse, bool bVelChange)
-{
-	check(GetBodyInstance());
-	GetBodyInstance()->AddImpulse(Impulse, bVelChange);
-
-#if !(UE_BUILD_SHIPPING || UE_BUILD_TEST)
-	if (GVehicleDebugParams.ShowAllForces)
-	{
-		FVector Position = VehicleState.VehicleWorldCOM;
-		DrawDebugDirectionalArrow(GetWorld(), Position, Position + Impulse * GVehicleDebugParams.ForceDebugScaling
-			, 20.f, FColor::Red, false, 0, 0, 2.f);
-	}
-#endif
-}
-
-void UChaosVehicleMovementComponent::AddImpulseAtPosition(const FVector& Impulse, const FVector& Position)
-{
-	check(GetBodyInstance());
-	GetBodyInstance()->AddImpulseAtPosition(Impulse, Position);
-
-#if !(UE_BUILD_SHIPPING || UE_BUILD_TEST)
-	if (GVehicleDebugParams.ShowAllForces)
-	{
-		DrawDebugDirectionalArrow(GetWorld(), Position, Position + Impulse * GVehicleDebugParams.ForceDebugScaling
-			, 20.f, FColor::Red, false, 0, 0, 2.f);
-	}
-#endif
-
-}
-
-void UChaosVehicleMovementComponent::AddTorqueInRadians(const FVector& Torque, bool bAllowSubstepping /*= true*/, bool bAccelChange /*= false*/)
-{
-	check(GetBodyInstance());
-	GetBodyInstance()->AddTorqueInRadians(Torque, bAllowSubstepping, bAccelChange);
-
-#if !(UE_BUILD_SHIPPING || UE_BUILD_TEST)
-	//#todo: how do we visualize torque?
-#endif
-}
-
-void UChaosVehicleMovementComponent::CreatePhysicsVehicle()
-{
-	PVehicle = MakeUnique<Chaos::FSimpleWheeledVehicle>();
-}
-
-
 void FVehicleAerofoilConfig::FillAerofoilSetup(const UChaosVehicleMovementComponent& MovementComponent)
 {
 	PAerofoilConfig.Type = (Chaos::EAerofoilType)(this->AerofoilType);
@@ -1389,4 +1487,275 @@
 	PThrusterConfig.MaxControlAngle = this->MaxControlAngle;
 }
 
+
+TUniquePtr<FChaosVehicleAsyncInput> UChaosVehicleMovementComponent::SetCurrentAsyncInputOutput(int32 InputIdx, FChaosVehicleManagerAsyncOutput* CurOutput, FChaosVehicleManagerAsyncOutput* NextOutput, float Alpha, int32 VehicleManagerTimestamp)
+{
+	TUniquePtr<FChaosVehicleDefaultAsyncInput> CurInput = MakeUnique<FChaosVehicleDefaultAsyncInput>();
+	SetCurrentAsyncInputOutputInternal(CurInput.Get(), InputIdx, CurOutput, NextOutput, Alpha, VehicleManagerTimestamp);
+	return CurInput;
+}
+
+
+/************************************************************************/
+/* Setup the current async I/O data                                     */
+/************************************************************************/
+void UChaosVehicleMovementComponent::SetCurrentAsyncInputOutputInternal(FChaosVehicleAsyncInput* CurInput, int32 InputIdx, FChaosVehicleManagerAsyncOutput* CurOutput, FChaosVehicleManagerAsyncOutput* NextOutput, float Alpha, int32 VehicleManagerTimestamp)
+{
+	ensure(CurAsyncInput == nullptr);	//should be reset after it was filled
+	ensure(CurAsyncOutput == nullptr);	//should get reset after update is done
+
+	CurAsyncInput = CurInput;
+	CurAsyncInput->Vehicle = this;
+	CurAsyncType = CurInput->Type;
+	NextAsyncOutput = nullptr;
+	OutputInterpAlpha = 0.f;
+
+	//FChaosVehicleDefaultAsyncInput* CurDefaultAsyncInput = static_cast<FChaosVehicleDefaultAsyncInput*>(CurAsyncInput);
+	//CurDefaultAsyncInput->ControlInputs.SteeringInput = 1.0f;
+
+	// We need to find our vehicle in the output given
+	if (CurOutput)
+	{
+		for (int32 PendingOutputIdx = 0; PendingOutputIdx < OutputsWaitingOn.Num(); ++PendingOutputIdx)
+		{
+			// Found the correct pending output, use index to get the vehicle.
+			if (OutputsWaitingOn[PendingOutputIdx].Timestamp == CurOutput->Timestamp)
+			{
+				const int32 VehicleIdx = OutputsWaitingOn[PendingOutputIdx].Idx;
+				FChaosVehicleAsyncOutput* VehicleOutput = CurOutput->VehicleOutputs[VehicleIdx].Get();
+				if (VehicleOutput && VehicleOutput->bValid && VehicleOutput->Type == CurAsyncType)
+				{
+					CurAsyncOutput = VehicleOutput;
+
+					if (NextOutput && NextOutput->Timestamp == CurOutput->Timestamp)
+					{
+						// This can occur when substepping - in this case, VehicleOutputs will be in the same order in NextOutput and CurOutput.
+						FChaosVehicleAsyncOutput* VehicleNextOutput = NextOutput->VehicleOutputs[VehicleIdx].Get();
+						if (VehicleNextOutput && VehicleNextOutput->bValid && VehicleNextOutput->Type == CurAsyncType)
+						{
+							NextAsyncOutput = VehicleNextOutput;
+							OutputInterpAlpha = Alpha;
+						}
+					}
+				}
+
+				// these are sorted by timestamp, we are using latest, so remove entries that came before it.
+				TArray<FAsyncOutputWrapper> NewOutputsWaitingOn;
+				for (int32 CopyIndex = PendingOutputIdx; CopyIndex < OutputsWaitingOn.Num(); ++CopyIndex)
+				{
+					NewOutputsWaitingOn.Add(OutputsWaitingOn[CopyIndex]);
+				}
+
+				OutputsWaitingOn = MoveTemp(NewOutputsWaitingOn);
+				break;
+			}
+		}
+
+	}
+
+	if (NextOutput && CurOutput)
+	{
+		if (NextOutput->Timestamp != CurOutput->Timestamp)
+		{
+			// NextOutput and CurOutput occurred in different steps, so we need to search for our specific vehicle.
+			for (int32 PendingOutputIdx = 0; PendingOutputIdx < OutputsWaitingOn.Num(); ++PendingOutputIdx)
+			{
+				// Found the correct pending output, use index to get the vehicle.
+				if (OutputsWaitingOn[PendingOutputIdx].Timestamp == NextOutput->Timestamp)
+				{
+					FChaosVehicleAsyncOutput* VehicleOutput = NextOutput->VehicleOutputs[OutputsWaitingOn[PendingOutputIdx].Idx].Get();
+					if (VehicleOutput && VehicleOutput->bValid && VehicleOutput->Type == CurAsyncType)
+					{
+						NextAsyncOutput = VehicleOutput;
+						OutputInterpAlpha = Alpha;
+					}
+					break;
+				}
+			}
+		}
+	}
+
+	FAsyncOutputWrapper& NewOutput = OutputsWaitingOn.AddDefaulted_GetRef();
+	NewOutput.Timestamp = VehicleManagerTimestamp;
+	NewOutput.Idx = InputIdx;
+}
+
+
+// ---- ASYNC ----
+
+FChaosVehicleDefaultAsyncInput::FChaosVehicleDefaultAsyncInput()
+	: GravityZ(0.f)
+{
+
+}
+
+/************************************************************************/
+/* Async simulation callback on the Physics Thread                      */
+/************************************************************************/
+TUniquePtr<FChaosVehicleAsyncOutput> FChaosVehicleDefaultAsyncInput::Simulate(UWorld* World, const float DeltaSeconds, const float TotalSeconds, bool& bWakeOut) const
+{
+	TUniquePtr<FChaosVehicleAsyncOutput> Output = MakeUnique<FChaosVehicleAsyncOutput>();
+
+	//UE_LOG(LogChaos, Warning, TEXT("Vehicle Physics Thread Tick %f"), DeltaSeconds);
+
+	//support nullptr because it allows us to go wide on filling the async inputs
+	if (Actor.Proxy == nullptr)
+	{
+		return Output;
+	}
+
+	// We now have access to the physics representation of the chassis on the physics thread async tick
+	Chaos::TPBDRigidParticleHandle<float, 3>* Handle = Actor.Proxy->GetHandle();
+
+	// FILL OUTPUT DATA HERE THAT WILL GET PASSED BACK TO THE GAME THREAD
+	Vehicle->VehicleSimulationPT->TickVehicle(World, DeltaSeconds, *this, *Output.Get(), Handle);
+
+	Output->bValid = true;
+
+	return MoveTemp(Output);
+}
+
+/************************************************************************/
+/* PASS ANY INUTS TO THE PHYSICS THREAD SIMULATION IN HERE              */
+/************************************************************************/
+void UChaosVehicleMovementComponent::Update(float DeltaTime)
+{
+	if (CurAsyncInput)
+	{
+		if (const FBodyInstance* BodyInstance = GetBodyInstance())
+		{
+			if (auto Handle = BodyInstance->ActorHandle)
+			{
+				CurAsyncInput->Actor.Proxy = static_cast<FRigidParticlePhysicsProxy*>(Handle->GetProxy());	// vehicles are never static
+
+				FChaosVehicleDefaultAsyncInput* AsyncInput = static_cast<FChaosVehicleDefaultAsyncInput*>(CurAsyncInput);
+				AsyncInput->ControlInputs.ThrottleInput = ThrottleInput;
+				AsyncInput->ControlInputs.BrakeInput = BrakeInput;
+				AsyncInput->ControlInputs.SteeringInput = SteeringInput;
+				AsyncInput->ControlInputs.HandbrakeInput = HandbrakeInput;
+				AsyncInput->ControlInputs.RollInput = RollInput;
+				AsyncInput->ControlInputs.PitchInput = PitchInput;
+				AsyncInput->ControlInputs.YawInput = YawInput;
+				AsyncInput->ControlInputs.GearUpInput = bRawGearUpInput;
+				AsyncInput->ControlInputs.GearDownInput = bRawGearDownInput;
+				AsyncInput->ControlInputs.TargetGearInput = TargetGearInput;
+				AsyncInput->ControlInputs.SetGearImmediate = SetGearImmediate;
+				AsyncInput->ControlInputs.TransmissionType = TransmissionType;
+
+				AsyncInput->GravityZ = GetGravityZ();
+
+				// #TODO: don't create every frame, setup never changes dynamically
+				TArray<AActor*> ActorsToIgnore;
+				ActorsToIgnore.Add(GetPawnOwner()); // ignore self in scene query
+
+				FCollisionQueryParams TraceParams(NAME_None, FCollisionQueryParams::GetUnknownStatId(), false, nullptr);
+				TraceParams.bReturnPhysicalMaterial = true;	// we need this to get the surface friction coefficient
+				TraceParams.AddIgnoredActors(ActorsToIgnore);
+				TraceParams.bTraceComplex = true;  // (Wheels[WheelIdx]->SweepType == ESweepType::ComplexSweep); FIX this again
+				AsyncInput->TraceParams = TraceParams;
+			}
+		}
+	}
+}
+
+void UChaosVehicleMovementComponent::FinalizeSimCallbackData(FChaosVehicleManagerAsyncInput& Input)
+{
+	bool bIsPhysicsStateCreated = true;
+	if (GetSkeletalMesh())
+	{
+		if (USkeletalMeshComponent* SkelMesh = GetSkeletalMesh())
+		{
+			bIsPhysicsStateCreated = SkelMesh->IsPhysicsStateCreated();
+		}
+	}
+
+	CurAsyncInput = nullptr;
+	CurAsyncOutput = nullptr;
+}
+
+
+/***************************************************************************/
+/* READ OUTPUT DATA - Access the async output data from the Physics Thread */
+/***************************************************************************/
+void UChaosVehicleMovementComponent::ParallelUpdate(float DeltaSeconds)
+{
+	if (const FChaosVehicleAsyncOutput* CurrentOutput = static_cast<FChaosVehicleAsyncOutput*>(CurAsyncOutput))
+	{
+		if (CurrentOutput->bValid)
+		{
+			// TODO: It would be nicer to go through CurAsyncOutput rather
+			// than copying into the vehicle, think about non-async path
+			PVehicleOutput->CurrentGear = CurAsyncOutput->VehicleSimOutput.CurrentGear;
+			PVehicleOutput->TargetGear = CurAsyncOutput->VehicleSimOutput.TargetGear;
+
+			if (TargetGearInput == PVehicleOutput->TargetGear)
+			{
+				SetGearImmediate = -1; // clear this otherwise can get gear change handshake issues between threads
+			}
+
+			// WHEN RUNNING WITH ASYNC ON & FIXED TIMESTEP THEN WE NEED TO INTERPOLATE BETWEEN THE CURRENT AND NEXT OUTPUT RESULTS
+			if (const FChaosVehicleAsyncOutput* NextOutput = static_cast<FChaosVehicleAsyncOutput*>(NextAsyncOutput))
+			{
+				PVehicleOutput->EngineRPM = FMath::Lerp(CurAsyncOutput->VehicleSimOutput.EngineRPM, NextAsyncOutput->VehicleSimOutput.EngineRPM, OutputInterpAlpha);
+				PVehicleOutput->EngineTorque = FMath::Lerp(CurAsyncOutput->VehicleSimOutput.EngineTorque, NextAsyncOutput->VehicleSimOutput.EngineTorque, OutputInterpAlpha);
+				PVehicleOutput->TransmissionRPM = FMath::Lerp(CurAsyncOutput->VehicleSimOutput.TransmissionRPM, NextAsyncOutput->VehicleSimOutput.TransmissionRPM, OutputInterpAlpha);
+				PVehicleOutput->TransmissionTorque = FMath::Lerp(CurAsyncOutput->VehicleSimOutput.TransmissionTorque, NextAsyncOutput->VehicleSimOutput.TransmissionTorque, OutputInterpAlpha);
+
+				for (int WheelIdx = 0; WheelIdx < CurrentOutput->VehicleSimOutput.Wheels.Num(); WheelIdx++)
+				{
+					const FWheelsOutput& Current = CurrentOutput->VehicleSimOutput.Wheels[WheelIdx];
+					const FWheelsOutput& Next = NextOutput->VehicleSimOutput.Wheels[WheelIdx];
+
+					PVehicleOutput->Wheels[WheelIdx].InContact = Current.InContact;
+					PVehicleOutput->Wheels[WheelIdx].SteeringAngle = FMath::Lerp(Current.SteeringAngle, Next.SteeringAngle, OutputInterpAlpha);
+					float DeltaAngle = FMath::FindDeltaAngleRadians(Current.AngularPosition, Next.AngularPosition);
+					PVehicleOutput->Wheels[WheelIdx].AngularPosition = Current.AngularPosition + DeltaAngle * OutputInterpAlpha;
+					PVehicleOutput->Wheels[WheelIdx].LateralAdhesiveLimit = FMath::Lerp(Current.LateralAdhesiveLimit, Next.LateralAdhesiveLimit, OutputInterpAlpha);
+					PVehicleOutput->Wheels[WheelIdx].LongitudinalAdhesiveLimit = FMath::Lerp(Current.LongitudinalAdhesiveLimit, Next.LongitudinalAdhesiveLimit, OutputInterpAlpha);
+
+					PVehicleOutput->Wheels[WheelIdx].bIsSlipping = Current.bIsSlipping;
+					PVehicleOutput->Wheels[WheelIdx].SlipMagnitude = FMath::Lerp(Current.SlipMagnitude, Next.SlipMagnitude, OutputInterpAlpha);
+					PVehicleOutput->Wheels[WheelIdx].bIsSkidding = Current.bIsSkidding;
+					PVehicleOutput->Wheels[WheelIdx].SkidMagnitude = FMath::Lerp(Current.SkidMagnitude, Next.SkidMagnitude, OutputInterpAlpha);
+					PVehicleOutput->Wheels[WheelIdx].SkidNormal = FMath::Lerp(Current.SkidNormal, Next.SkidNormal, OutputInterpAlpha);
+
+					PVehicleOutput->Wheels[WheelIdx].SuspensionOffset = FMath::Lerp(Current.SuspensionOffset, Next.SuspensionOffset, OutputInterpAlpha);
+					PVehicleOutput->Wheels[WheelIdx].SpringForce = FMath::Lerp(Current.SpringForce, Next.SpringForce, OutputInterpAlpha);
+					PVehicleOutput->Wheels[WheelIdx].NormalizedSuspensionLength = FMath::Lerp(Current.NormalizedSuspensionLength, Next.NormalizedSuspensionLength, OutputInterpAlpha);
+				}
+			}
+			else // WHEN ASYNC IS OFF IT STILL GENERATES THE ASYNC CALLBACK BUT THERE IS ONLY EVER THE CURRENT AND NO NEXT OUTPUT TO INTERPOLATE BETWEEN
+			{
+				PVehicleOutput->EngineRPM = CurAsyncOutput->VehicleSimOutput.EngineRPM;
+				PVehicleOutput->EngineTorque = CurAsyncOutput->VehicleSimOutput.EngineTorque;
+				PVehicleOutput->TransmissionRPM = CurAsyncOutput->VehicleSimOutput.TransmissionRPM;
+				PVehicleOutput->TransmissionTorque = CurAsyncOutput->VehicleSimOutput.TransmissionTorque;
+
+				for (int WheelIdx = 0; WheelIdx < CurrentOutput->VehicleSimOutput.Wheels.Num(); WheelIdx++)
+				{
+					const FWheelsOutput& Current = CurrentOutput->VehicleSimOutput.Wheels[WheelIdx];
+
+					PVehicleOutput->Wheels[WheelIdx].InContact = Current.InContact;
+					PVehicleOutput->Wheels[WheelIdx].SteeringAngle = Current.SteeringAngle;
+					PVehicleOutput->Wheels[WheelIdx].AngularPosition = Current.AngularPosition;
+					PVehicleOutput->Wheels[WheelIdx].LateralAdhesiveLimit = Current.LateralAdhesiveLimit;
+					PVehicleOutput->Wheels[WheelIdx].LongitudinalAdhesiveLimit = Current.LongitudinalAdhesiveLimit;
+
+					PVehicleOutput->Wheels[WheelIdx].bIsSlipping = Current.bIsSlipping;
+					PVehicleOutput->Wheels[WheelIdx].SlipMagnitude = Current.SlipMagnitude;
+					PVehicleOutput->Wheels[WheelIdx].bIsSkidding = Current.bIsSkidding;
+					PVehicleOutput->Wheels[WheelIdx].SkidMagnitude = Current.SkidMagnitude;
+					PVehicleOutput->Wheels[WheelIdx].SkidNormal = Current.SkidNormal;
+
+					PVehicleOutput->Wheels[WheelIdx].SuspensionOffset = Current.SuspensionOffset;
+					PVehicleOutput->Wheels[WheelIdx].SpringForce = Current.SpringForce;
+					PVehicleOutput->Wheels[WheelIdx].NormalizedSuspensionLength = Current.NormalizedSuspensionLength;
+				}
+
+			}
+
+		}
+	}
+}
+
 #undef LOCTEXT_NAMESPACE
@@ -1396,3 +1765,4 @@
 PRAGMA_ENABLE_OPTIMIZATION
 #endif
 
+
diff --git a/Engine/Plugins/Experimental/ChaosVehiclesPlugin/Source/ChaosVehicles/Private/ChaosVehicleWheel.cpp b/Engine/Plugins/Experimental/ChaosVehiclesPlugin/Source/ChaosVehicles/Private/ChaosVehicleWheel.cpp
--- a/Engine/Plugins/Experimental/ChaosVehiclesPlugin/Source/ChaosVehicles/Private/ChaosVehicleWheel.cpp
+++ b/Engine/Plugins/Experimental/ChaosVehiclesPlugin/Source/ChaosVehicles/Private/ChaosVehicleWheel.cpp
@@ -53,37 +53,40 @@ UChaosVehicleWheel::UChaosVehicleWheel(const FObjectInitializer& ObjectInitializ
 
 FChaosVehicleManager* UChaosVehicleWheel::GetVehicleManager() const
 {
-	UWorld* World = GEngine->GetWorldFromContextObject(VehicleSim, EGetWorldErrorMode::LogAndReturnNull);
+	UWorld* World = GEngine->GetWorldFromContextObject(VehicleComponent, EGetWorldErrorMode::LogAndReturnNull);
 	return World ? FChaosVehicleManager::GetVehicleManagerFromScene(World->GetPhysicsScene()) : nullptr;
 }
 
 
 float UChaosVehicleWheel::GetSteerAngle() const
 {
-	check(VehicleSim && VehicleSim->PhysicsVehicle());
-	return VehicleSim->PhysicsVehicle()->Wheels[WheelIndex].GetSteeringAngle();
+	check(VehicleComponent && VehicleComponent->PhysicsVehicleOutput());
+	return VehicleComponent->PhysicsVehicleOutput()->Wheels[WheelIndex].SteeringAngle;
 }
 
 float UChaosVehicleWheel::GetRotationAngle() const
 {
-	check(VehicleSim && VehicleSim->PhysicsVehicle());
-	float RotationAngle = -1.0f * FMath::RadiansToDegrees(VehicleSim->PhysicsVehicle()->Wheels[WheelIndex].GetAngularPosition());
+	check(VehicleComponent && VehicleComponent->PhysicsVehicleOutput());
+	float RotationAngle = -1.0f * FMath::RadiansToDegrees(VehicleComponent->PhysicsVehicleOutput()->Wheels[WheelIndex].AngularPosition);
 	ensure(!FMath::IsNaN(RotationAngle));
 
+	//if (WheelIndex == 2)
+	//{
+	//	UE_LOG(LogChaos, Warning, TEXT("RotationAngle = %f"), RotationAngle);
+	//}
 	return RotationAngle;
 }
 
 float UChaosVehicleWheel::GetSuspensionOffset() const
 {
-	check(VehicleSim && VehicleSim->PhysicsVehicle());
-	Chaos::FSimpleSuspensionSim& SimSuspension = VehicleSim->PhysicsVehicle()->Suspension[WheelIndex];
-	return SimSuspension.GetSuspensionOffset();
+	check(VehicleComponent && VehicleComponent->PhysicsVehicleOutput());
+	return VehicleComponent->PhysicsVehicleOutput()->Wheels[WheelIndex].SuspensionOffset;
 }
 
 bool UChaosVehicleWheel::IsInAir() const
 {
-	check(VehicleSim && VehicleSim->PhysicsVehicle());
-	return !VehicleSim->PhysicsVehicle()->Wheels[WheelIndex].InContact();
+	check(VehicleComponent && VehicleComponent->PhysicsVehicleOutput());
+	return !VehicleComponent->PhysicsVehicleOutput()->Wheels[WheelIndex].InContact;
 }
 
 
@@ -92,7 +95,7 @@ void UChaosVehicleWheel::Init( UChaosWheeledVehicleMovementComponent* InVehicleS
 	check(InVehicleSim);
 	check(InVehicleSim->Wheels.IsValidIndex(InWheelIndex));
 
-	VehicleSim = InVehicleSim;
+	VehicleComponent = InVehicleSim;
 	WheelIndex = InWheelIndex;
 
 //#if WITH_PHYSX_VEHICLES
@@ -119,7 +122,7 @@ void UChaosVehicleWheel::Shutdown()
 
 FChaosWheelSetup& UChaosVehicleWheel::GetWheelSetup()
 {
-	return VehicleSim->WheelSetups[WheelIndex];
+	return VehicleComponent->WheelSetups[WheelIndex];
 }
 
 void UChaosVehicleWheel::Tick( float DeltaTime )
diff --git a/Engine/Plugins/Experimental/ChaosVehiclesPlugin/Source/ChaosVehicles/Private/ChaosWheeledVehicleMovementComponent.cpp b/Engine/Plugins/Experimental/ChaosVehiclesPlugin/Source/ChaosVehicles/Private/ChaosWheeledVehicleMovementComponent.cpp
--- a/Engine/Plugins/Experimental/ChaosVehiclesPlugin/Source/ChaosVehicles/Private/ChaosWheeledVehicleMovementComponent.cpp
+++ b/Engine/Plugins/Experimental/ChaosVehiclesPlugin/Source/ChaosVehicles/Private/ChaosWheeledVehicleMovementComponent.cpp
@@ -87,6 +87,785 @@
 	LocalWheelVelocity[WheelIdx] = WorldTransform.InverseTransformVector(WorldWheelVelocity[WheelIdx]);
 }
 
+void FWheelState::CaptureState(int WheelIdx, const FVector& WheelOffset, const Chaos::TPBDRigidParticleHandle<float, 3>* Handle)
+{
+	check(Handle);
+	const FTransform WorldTransform(Handle->R(), Handle->X());
+	WheelWorldLocation[WheelIdx] = WorldTransform.TransformPosition(WheelOffset);
+	WorldWheelVelocity[WheelIdx] = GetVelocityAtPoint(Handle, WheelWorldLocation[WheelIdx]);
+	LocalWheelVelocity[WheelIdx] = WorldTransform.InverseTransformVector(WorldWheelVelocity[WheelIdx]);
+}
+
+FVector FWheelState::GetVelocityAtPoint(const Chaos::TPBDRigidParticleHandle<float, 3>* Rigid, const FVector& InPoint)
+{
+	const Chaos::FVec3 COM = Rigid ? Chaos::FParticleUtilitiesGT::GetCoMWorldPosition(Rigid) : Chaos::FParticleUtilitiesGT::GetActorWorldTransform(Rigid).GetTranslation();
+	const Chaos::FVec3 Diff = InPoint - COM;
+	return Rigid->V() - Chaos::FVec3::CrossProduct(Diff, Rigid->W());
+}
+
+/**
+ * UChaosWheeledVehicleSimulation
+ */
+bool UChaosWheeledVehicleSimulation::CanSimulate() const
+{
+	if (UChaosVehicleSimulation::CanSimulate() == false)
+	{
+		return false;
+	}
+
+	return (PVehicle && PVehicle.IsValid()
+		&& PVehicle->Engine.Num() == PVehicle->Transmission.Num()
+		&& Wheels.Num() > 0 && Wheels.Num() == PVehicle->Suspension.Num());
+}
+
+void UChaosWheeledVehicleSimulation::TickVehicle(UWorld* WorldIn, float DeltaTime, const FChaosVehicleDefaultAsyncInput& InputData, FChaosVehicleAsyncOutput& OutputData, Chaos::TPBDRigidParticleHandle<float, 3>* Handle)
+{
+	UChaosVehicleSimulation::TickVehicle(WorldIn, DeltaTime, InputData, OutputData, Handle);
+}
+
+void UChaosWheeledVehicleSimulation::UpdateSimulation(float DeltaTime, const FChaosVehicleDefaultAsyncInput& InputData, Chaos::TPBDRigidParticleHandle<float, 3>* Handle)
+{
+	// Inherit common vehicle simulation stages ApplyAerodynamics, ApplyTorqueControl, etc
+	UChaosVehicleSimulation::UpdateSimulation(DeltaTime, InputData, Handle);
+
+	if (CanSimulate() && Handle)
+	{
+		// sanity check that everything is setup ok
+		ensure(Wheels.Num() == PVehicle->Suspension.Num());
+		ensure(Wheels.Num() == PVehicle->Wheels.Num());
+		ensure(WheelState.LocalWheelVelocity.Num() == Wheels.Num());
+		ensure(WheelState.WheelWorldLocation.Num() == Wheels.Num());
+		ensure(WheelState.WorldWheelVelocity.Num() == Wheels.Num());
+
+		///////////////////////////////////////////////////////////////////////
+		// Cache useful state so we are not re-calculating the same data
+
+		for (int WheelIdx = 0; WheelIdx < Wheels.Num(); WheelIdx++)
+		{
+			WheelState.CaptureState(WheelIdx, PVehicle->Suspension[WheelIdx].GetLocalRestingPosition(), Handle);
+		}
+
+		///////////////////////////////////////////////////////////////////////
+		// Suspension Raycast
+
+		for (int WheelIdx = 0; WheelIdx < Wheels.Num(); WheelIdx++)
+		{
+			auto& PSuspension = PVehicle->Suspension[WheelIdx];
+			auto& PWheel = PVehicle->Wheels[WheelIdx];
+			PSuspension.UpdateWorldRaycastLocation(VehicleState.VehicleWorldTransform, PWheel.Setup().WheelRadius, WheelState.Trace[WheelIdx]);
+		}
+
+		if (!GWheeledVehicleDebugParams.DisableSuspensionForces && PVehicle->bSuspensionEnabled)
+		{
+			PerformSuspensionTraces(WheelState.Trace, InputData.TraceParams);
+		}
+
+		//////////////////////////////////////////////////////////////////////////
+		// Wheel and Vehicle in air state
+
+		VehicleState.bVehicleInAir = true;
+		VehicleState.NumWheelsOnGround = 0;
+		for (int WheelIdx = 0; WheelIdx < Wheels.Num(); WheelIdx++)
+		{
+			// tell systems who care that wheel is touching the ground
+			PVehicle->Wheels[WheelIdx].SetOnGround(Wheels[WheelIdx]->HitResult.bBlockingHit);
+
+			// only requires one wheel to be on the ground for the vehicle to be NOT in the air
+			if (PVehicle->Wheels[WheelIdx].InContact())
+			{
+				VehicleState.bVehicleInAir = false;
+				VehicleState.NumWheelsOnGround++;
+			}
+		}
+		VehicleState.bAllWheelsOnGround = (VehicleState.NumWheelsOnGround == Wheels.Num());
+
+		///////////////////////////////////////////////////////////////////////
+		// Input
+		ApplyInput(InputData.ControlInputs, DeltaTime);
+
+		///////////////////////////////////////////////////////////////////////
+		// Engine/Transmission
+		if (!GWheeledVehicleDebugParams.DisableSuspensionForces && PVehicle->bMechanicalSimEnabled)
+		{
+			ProcessMechanicalSimulation(DeltaTime);
+		}
+
+		///////////////////////////////////////////////////////////////////////
+		// Suspension
+
+		if (!GWheeledVehicleDebugParams.DisableSuspensionForces && PVehicle->bSuspensionEnabled)
+		{
+			ApplySuspensionForces(DeltaTime);
+		}
+
+		///////////////////////////////////////////////////////////////////////
+		// Steering
+
+		ProcessSteering(InputData.ControlInputs);
+
+		///////////////////////////////////////////////////////////////////////
+		// Wheel Friction
+
+		if (!GWheeledVehicleDebugParams.DisableFrictionForces && PVehicle->bWheelFrictionEnabled)
+		{
+			ApplyWheelFrictionForces(DeltaTime);
+		}
+
+#if 0
+		if (PerformanceMeasure.IsEnabled())
+		{
+			PerformanceMeasure.Update(DeltaTime, VehicleState.VehicleWorldTransform.GetLocation(), VehicleState.ForwardSpeed);
+		}
+#endif
+	}
+
+}
+
+void UChaosWheeledVehicleSimulation::PerformSuspensionTraces(const TArray<FSuspensionTrace>& SuspensionTrace, const FCollisionQueryParams& TraceParams)
+{
+	//#todo: SpringCollisionChannel should be a parameter setup
+	ECollisionChannel SpringCollisionChannel = ECollisionChannel::ECC_WorldDynamic;
+	FCollisionResponseParams ResponseParams;
+
+	// batching is about 0.5ms (25%) faster when there's 100 vehicles on a flat terrain
+	if (GVehicleDebugParams.BatchQueries)
+	{
+		FBox QueryBox;
+		for (int WheelIdx = 0; WheelIdx < Wheels.Num(); WheelIdx++)
+		{
+
+			FVector TraceStart = SuspensionTrace[WheelIdx].Start;
+			FVector TraceEnd = SuspensionTrace[WheelIdx].End;
+
+			if (WheelIdx == 0)
+			{
+				QueryBox = FBox(TraceStart, TraceEnd);
+			}
+			else
+			{
+				QueryBox.Min = QueryBox.Min.ComponentMin(TraceStart);
+				QueryBox.Min = QueryBox.Min.ComponentMin(TraceEnd);
+				QueryBox.Max = QueryBox.Max.ComponentMax(TraceStart);
+				QueryBox.Max = QueryBox.Max.ComponentMax(TraceEnd);
+			}
+		}
+		float OneWheelRadius = PVehicle->Wheels[0].GetEffectiveRadius(); // or wheel width
+
+		QueryBox.ExpandBy(FVector(OneWheelRadius, OneWheelRadius, OneWheelRadius)); // little extra just to be on the safe side consider 1 or 2 wheel vehicle
+		TArray<FOverlapResult> OverlapResults;
+		FCollisionShape CollisionBox;
+		CollisionBox.SetBox(QueryBox.GetExtent());
+
+#if !(UE_BUILD_SHIPPING || UE_BUILD_TEST)
+		if (GWheeledVehicleDebugParams.ShowBatchQueryExtents)
+		{
+			DrawDebugBox(World, QueryBox.GetCenter(), QueryBox.GetExtent(), FColor::Yellow, false, -1.0f, 0, 2.0f);
+		}
+#endif
+
+		const bool bOverlapHit = World->OverlapMultiByChannel(OverlapResults, QueryBox.GetCenter(), FQuat::Identity, SpringCollisionChannel, CollisionBox, TraceParams, ResponseParams);
+
+		for (int32 WheelIdx = 0; WheelIdx < Wheels.Num(); ++WheelIdx)
+		{
+			FHitResult& HitResult = Wheels[WheelIdx]->HitResult;
+			HitResult = FHitResult();
+
+			if (bOverlapHit)
+			{
+				const FVector& TraceStart = SuspensionTrace[WheelIdx].Start;
+				const FVector& TraceEnd = SuspensionTrace[WheelIdx].End;
+				//TraceParams.bTraceComplex = (Wheels[WheelIdx]->SweepType == ESweepType::ComplexSweep);
+
+				FVector TraceVector(TraceStart - TraceEnd); // reversed
+				FVector TraceNormal = TraceVector.GetSafeNormal();
+
+				// Test each overlapped object for a hit result
+				for (FOverlapResult OverlapResult : OverlapResults)
+				{
+					if (!OverlapResult.bBlockingHit)
+						continue;
+
+					FHitResult ComponentHit;
+
+					switch (Wheels[WheelIdx]->SweepShape)
+					{
+					case ESweepShape::Spherecast:
+					{
+						float WheelRadius = PVehicle->Wheels[WheelIdx].GetEffectiveRadius(); // or wheel width
+						FVector VehicleUpAxis = TraceNormal;// GetOwner()->GetTransform().GetUnitAxis(EAxis::Z);
+
+						FVector Start = TraceStart + VehicleUpAxis * WheelRadius;
+						FVector End = TraceEnd + VehicleUpAxis * WheelRadius;
+
+						if (OverlapResult.Component->SweepComponent(ComponentHit, Start, End, FQuat::Identity, FCollisionShape::MakeSphere(WheelRadius), TraceParams.bTraceComplex))
+						{
+							if (ComponentHit.Time < HitResult.Time)
+							{
+								HitResult = ComponentHit;
+								HitResult.bBlockingHit = OverlapResult.bBlockingHit;
+							}
+						}
+					}
+					break;
+
+					case ESweepShape::Raycast:
+					default:
+					{
+						if (OverlapResult.Component->LineTraceComponent(ComponentHit, TraceStart, TraceEnd, TraceParams))
+						{
+							if (ComponentHit.Time < HitResult.Time)
+							{
+								HitResult = ComponentHit;
+								HitResult.bBlockingHit = OverlapResult.bBlockingHit;
+							}
+						}
+					}
+					break;
+					}
+				}
+			}
+		}
+	}
+	else
+	{
+		for (int WheelIdx = 0; WheelIdx < Wheels.Num(); WheelIdx++)
+		{
+			FHitResult& HitResult = Wheels[WheelIdx]->HitResult;
+
+			FVector TraceStart = SuspensionTrace[WheelIdx].Start;
+			FVector TraceEnd = SuspensionTrace[WheelIdx].End;
+			//TraceParams.bTraceComplex = (Wheels[WheelIdx]->SweepType == ESweepType::ComplexSweep);
+
+			FVector TraceVector(TraceStart - TraceEnd); // reversed
+			FVector TraceNormal = TraceVector.GetSafeNormal();
+
+			switch (Wheels[WheelIdx]->SweepShape)
+			{
+			case ESweepShape::Spherecast:
+			{
+				//float Radius = PVehicle->Wheels[WheelIdx].GetEffectiveRadius(); // or wheel width
+				float Radius = PVehicle->Wheels[WheelIdx].Setup().WheelWidth * 0.5f; // or wheel width
+				FVector VehicleUpAxis = TraceNormal; //GetOwner()->GetTransform().GetUnitAxis(EAxis::Z);
+
+				World->SweepSingleByChannel(HitResult
+					, TraceStart + VehicleUpAxis * Radius
+					, TraceEnd + VehicleUpAxis * Radius
+					, FQuat::Identity, SpringCollisionChannel
+					, FCollisionShape::MakeSphere(Radius), TraceParams
+					, FCollisionResponseParams::DefaultResponseParam);
+			}
+			break;
+
+
+
+			case ESweepShape::Raycast:
+			default:
+			{
+				World->LineTraceSingleByChannel(HitResult, TraceStart, TraceEnd, SpringCollisionChannel, TraceParams, FCollisionResponseParams::DefaultResponseParam);
+			}
+			break;
+			}
+		}
+	}
+
+}
+
+void UChaosWheeledVehicleSimulation::ApplyWheelFrictionForces(float DeltaTime)
+{
+	for (int WheelIdx = 0; WheelIdx < Wheels.Num(); WheelIdx++)
+	{
+		auto& PWheel = PVehicle->Wheels[WheelIdx]; // Physics Wheel
+		FHitResult& HitResult = Wheels[WheelIdx]->HitResult;
+		UChaosVehicleWheel* VehicleWheel = Wheels[WheelIdx];
+
+		if (PWheel.InContact())
+		{
+			if (HitResult.PhysMaterial.IsValid())
+			{
+				PWheel.SetSurfaceFriction(HitResult.PhysMaterial->Friction);
+			}
+
+			// take into account steering angle
+			float SteerAngleDegrees = VehicleWheel->GetSteerAngle(); 
+			FRotator SteeringRotator(0.f, SteerAngleDegrees, 0.f);
+			FVector SteerLocalWheelVelocity = SteeringRotator.UnrotateVector(WheelState.LocalWheelVelocity[WheelIdx]);
+
+			PWheel.SetVehicleGroundSpeed(SteerLocalWheelVelocity);
+			PWheel.Simulate(DeltaTime);
+
+			float RotationAngle = FMath::RadiansToDegrees(PWheel.GetAngularPosition());
+			FVector FrictionForceLocal = PWheel.GetForceFromFriction();
+			FrictionForceLocal = SteeringRotator.RotateVector(FrictionForceLocal);
+
+			FVector GroundZVector = HitResult.Normal;
+			FVector GroundXVector = FVector::CrossProduct(VehicleState.VehicleRightAxis, GroundZVector);
+			FVector GroundYVector = FVector::CrossProduct(GroundZVector, GroundXVector);
+
+			// the force should be applied along the ground surface not along vehicle forward vector?
+			//FVector FrictionForceVector = VehicleState.VehicleWorldTransform.TransformVector(FrictionForceLocal);
+			FMatrix Mat(GroundXVector, GroundYVector, GroundZVector, VehicleState.VehicleWorldTransform.GetLocation());
+			FVector FrictionForceVector = Mat.TransformVector(FrictionForceLocal);
+
+			check(PWheel.InContact());
+			AddForceAtPosition(FrictionForceVector, WheelState.WheelWorldLocation[WheelIdx]);
+
+#if !(UE_BUILD_SHIPPING || UE_BUILD_TEST)
+			if (GWheeledVehicleDebugParams.ShowWheelForces)
+			{
+				// show longitudinal drive force
+				{
+					DrawDebugLine(World
+						, WheelState.WheelWorldLocation[WheelIdx]
+						, WheelState.WheelWorldLocation[WheelIdx] + FrictionForceVector * 0.001f
+						, FColor::Yellow, false, -1.0f, 0, 2);
+
+					DrawDebugLine(World
+						, WheelState.WheelWorldLocation[WheelIdx]
+						, WheelState.WheelWorldLocation[WheelIdx] + GroundZVector * 100.f
+						, FColor::Orange, false, -1.0f, 0, 2);
+
+				}
+
+			}
+#endif
+
+		}
+		else
+		{
+			PWheel.SetVehicleGroundSpeed(WheelState.LocalWheelVelocity[WheelIdx]);
+			PWheel.Simulate(DeltaTime);
+		}
+
+	}
+}
+
+void UChaosWheeledVehicleSimulation::ApplySuspensionForces(float DeltaTime)
+{
+	TArray<float> SusForces;
+	SusForces.Init(0.f, Wheels.Num());
+
+	for (int WheelIdx = 0; WheelIdx < Wheels.Num(); WheelIdx++)
+	{
+		FHitResult& HitResult = Wheels[WheelIdx]->HitResult;
+
+		float NewDesiredLength = 1.0f; // suspension max length
+		float ForceMagnitude2 = 0.f;
+		auto& PWheel = PVehicle->Wheels[WheelIdx];
+		auto& PSuspension = PVehicle->Suspension[WheelIdx];
+		float SuspensionMovePosition = -PSuspension.Setup().MaxLength;
+
+		if (!GWheeledVehicleDebugParams.DisableSuspensionConstraint)
+		{
+#if WITH_CHAOS
+			if (ConstraintHandles.Num() > 0)
+			{
+				FPhysicsConstraintHandle& ConstraintHandle = ConstraintHandles[WheelIdx];
+				if (ConstraintHandle.IsValid())
+				{
+					if (Chaos::FSuspensionConstraint* Constraint = static_cast<Chaos::FSuspensionConstraint*>(ConstraintHandle.Constraint))
+					{
+						FVec3 P = HitResult.ImpactPoint + (PWheel.Setup().WheelRadius * VehicleState.VehicleUpAxis);
+						Constraint->SetTarget(P);
+						Constraint->SetEnabled(PWheel.InContact());
+					}
+				}
+
+			}
+#endif // WITH_CHAOS
+		}
+
+		if (PWheel.InContact())
+		{
+			NewDesiredLength = HitResult.Distance;
+
+			SuspensionMovePosition = -FVector::DotProduct(WheelState.WheelWorldLocation[WheelIdx] - HitResult.ImpactPoint, VehicleState.VehicleUpAxis) + PWheel.Setup().WheelRadius;
+
+			PSuspension.SetSuspensionLength(NewDesiredLength, PWheel.Setup().WheelRadius);
+			PSuspension.SetLocalVelocity(WheelState.LocalWheelVelocity[WheelIdx]);
+			PSuspension.Simulate(DeltaTime);
+
+			float ForceMagnitude = PSuspension.GetSuspensionForce();
+
+			FVector GroundZVector = HitResult.Normal;
+			FVector SuspensionForceVector = VehicleState.VehicleUpAxis * ForceMagnitude;
+
+			FVector SusApplicationPoint = WheelState.WheelWorldLocation[WheelIdx] + PVehicle->Suspension[WheelIdx].Setup().SuspensionForceOffset;
+
+			check(PWheel.InContact());
+			if (GWheeledVehicleDebugParams.DisableSuspensionConstraint)
+			{
+				AddForceAtPosition(SuspensionForceVector, SusApplicationPoint);
+			}
+
+			ForceMagnitude = PSuspension.Setup().WheelLoadRatio * ForceMagnitude + (1.f - PSuspension.Setup().WheelLoadRatio) * PSuspension.Setup().RestingForce;
+			PWheel.SetWheelLoadForce(ForceMagnitude);
+			PWheel.SetMassPerWheel(RigidHandle->M() / PVehicle->Wheels.Num());
+			SusForces[WheelIdx] = ForceMagnitude;
+
+#if !(UE_BUILD_SHIPPING || UE_BUILD_TEST)
+			if (GWheeledVehicleDebugParams.ShowSuspensionForces)
+			{
+				DrawDebugLine(World
+					, SusApplicationPoint
+					, SusApplicationPoint + SuspensionForceVector * GVehicleDebugParams.ForceDebugScaling
+					, FColor::Blue, false, -1.0f, 0, 5);
+
+				DrawDebugLine(World
+					, SusApplicationPoint
+					, SusApplicationPoint + GroundZVector * 140.f
+					, FColor::Yellow, false, -1.0f, 0, 5);
+			}
+#endif
+
+		}
+		else
+		{
+			PSuspension.SetSuspensionLength(PSuspension.GetTraceLength(PWheel.Setup().WheelRadius), PWheel.Setup().WheelRadius);
+		}
+
+	}
+
+	if (!GWheeledVehicleDebugParams.DisableRollbarForces)
+	{
+		for (auto& Axle : PVehicle->GetAxles())
+		{
+			//#todo: only works with 2 wheels on an axle at present
+			if (Axle.Setup.WheelIndex.Num() == 2)
+			{
+				uint16 WheelIdxA = Axle.Setup.WheelIndex[0];
+				uint16 WheelIdxB = Axle.Setup.WheelIndex[1];
+
+				float FV = Axle.Setup.RollbarScaling;
+				float ForceDiffOnAxleF = SusForces[WheelIdxA] - SusForces[WheelIdxB];
+				FVector ForceVector0 = VehicleState.VehicleUpAxis * ForceDiffOnAxleF * FV;
+				FVector ForceVector1 = VehicleState.VehicleUpAxis * ForceDiffOnAxleF * -FV;
+
+				FVector SusApplicationPoint0 = WheelState.WheelWorldLocation[WheelIdxA] + PVehicle->Suspension[WheelIdxA].Setup().SuspensionForceOffset;
+				AddForceAtPosition(ForceVector0, SusApplicationPoint0);
+
+				FVector SusApplicationPoint1 = WheelState.WheelWorldLocation[WheelIdxB] + PVehicle->Suspension[WheelIdxB].Setup().SuspensionForceOffset;
+				AddForceAtPosition(ForceVector1, SusApplicationPoint1);
+			}
+		}
+	}
+}
+
+void UChaosWheeledVehicleSimulation::ProcessSteering(const FControlInputs& ControlInputs)
+{
+	auto& PSteering = PVehicle->GetSteering();
+
+	for (int WheelIdx = 0; WheelIdx < Wheels.Num(); WheelIdx++)
+	{
+		auto& PWheel = PVehicle->Wheels[WheelIdx]; // Physics Wheel
+		FHitResult& HitResult = Wheels[WheelIdx]->HitResult;
+
+		if (PWheel.Setup().SteeringEnabled)
+		{
+			// #TODO: PUT BACK
+			//FRichCurve* SteeringCurveData = SteeringSetup.SteeringCurve.GetRichCurve();
+			float SpeedScale = 1.0f;
+
+			// #TODO: PUT BACK
+			//// allow full counter steering when steering into a power slide
+			////if (SteeringInput * VehicleState.VehicleLocalVelocity.Y > 0.0f)
+			//{
+			//	SpeedScale = SteeringCurveData->Eval(CmSToMPH(VehicleState.ForwardSpeed));
+			//}
+
+			float UseSteeringValue = ControlInputs.SteeringInput * SpeedScale;
+
+			float SteeringAngle = 0.f;
+			if (FMath::Abs(GWheeledVehicleDebugParams.SteeringOverride) > 0.01f)
+			{
+				SteeringAngle = PWheel.Setup().MaxSteeringAngle * GWheeledVehicleDebugParams.SteeringOverride;
+			}
+			else
+			{
+				float WheelSide = PVehicle->GetSuspension(WheelIdx).GetLocalRestingPosition().Y;
+				SteeringAngle = PSteering.GetSteeringAngle(UseSteeringValue, PWheel.Setup().MaxSteeringAngle, WheelSide);
+			}
+
+			PWheel.SetSteeringAngle(SteeringAngle);
+		}
+		else
+		{
+			PWheel.SetSteeringAngle(0.0f);
+		}
+	}
+}
+
+void UChaosWheeledVehicleSimulation::ApplyInput(const FControlInputs& ControlInputs, float DeltaTime)
+{
+	UChaosVehicleSimulation::ApplyInput(ControlInputs, DeltaTime);
+
+	FControlInputs ModifiedInputs = ControlInputs;
+
+	float EngineBraking = 0.f;
+	if (PVehicle->bMechanicalSimEnabled)
+	{
+		auto& PEngine = PVehicle->GetEngine();
+		auto& PTransmission = PVehicle->GetTransmission();
+
+		if (ModifiedInputs.TransmissionType != PTransmission.Setup().TransmissionType)
+		{
+			PTransmission.AccessSetup().TransmissionType = ModifiedInputs.TransmissionType;
+		}
+
+		if (ModifiedInputs.SetGearImmediate >= 0 && ModifiedInputs.TargetGearInput != PTransmission.GetTargetGear())
+		{
+			PTransmission.SetGear(ModifiedInputs.TargetGearInput, ModifiedInputs.SetGearImmediate?true:false);
+		}
+
+		if (ModifiedInputs.GearUpInput)
+		{
+			PTransmission.ChangeUp();
+			ModifiedInputs.GearUpInput = false;
+		}
+
+		if (ModifiedInputs.GearDownInput)
+		{
+			PTransmission.ChangeDown();
+			ModifiedInputs.GearDownInput = false;
+		}
+
+		if (GWheeledVehicleDebugParams.ThrottleOverride > 0.f)
+		{
+			PTransmission.SetGear(1, true);
+			ModifiedInputs.BrakeInput = 0.f;
+			PEngine.SetThrottle(GWheeledVehicleDebugParams.ThrottleOverride);
+		}
+		else
+		{
+			PEngine.SetThrottle(ModifiedInputs.ThrottleInput * ModifiedInputs.ThrottleInput);
+		}
+
+		EngineBraking = PEngine.GetEngineRPM() * PEngine.Setup().EngineBrakeEffect;
+	}
+
+	for (int WheelIdx = 0; WheelIdx < Wheels.Num(); WheelIdx++)
+	{
+		auto& PWheel = PVehicle->Wheels[WheelIdx];
+
+		float EngineBrakingForce = 0.0f;
+		if ((ModifiedInputs.ThrottleInput < SMALL_NUMBER) && FMath::Abs(VehicleState.ForwardSpeed) > SMALL_NUMBER && PWheel.Setup().EngineEnabled)
+		{
+			EngineBrakingForce = EngineBraking;
+		}
+
+		if (PWheel.Setup().BrakeEnabled)
+		{
+			float BrakeForce = PWheel.Setup().MaxBrakeTorque * ModifiedInputs.BrakeInput;
+			PWheel.SetBrakeTorque(MToCm(BrakeForce + EngineBrakingForce));
+		}
+		else
+		{
+			PWheel.SetBrakeTorque(MToCm(EngineBraking));
+		}
+
+		if (ModifiedInputs.HandbrakeInput && PWheel.Setup().HandbrakeEnabled)
+		{
+			PWheel.SetBrakeTorque(MToCm(ModifiedInputs.HandbrakeInput * PWheel.Setup().HandbrakeTorque));
+		}
+	}
+
+
+}
+
+bool UChaosWheeledVehicleSimulation::IsWheelSpinning() const
+{
+	for (auto& Wheel : PVehicle->Wheels)
+	{
+		if (Wheel.IsSlipping())
+		{
+			return true;
+		}
+	}
+
+	return false;
+}
+
+void UChaosWheeledVehicleSimulation::ProcessMechanicalSimulation(float DeltaTime)
+{
+	auto& PEngine = PVehicle->GetEngine();
+	auto& PTransmission = PVehicle->GetTransmission();
+	auto& PDifferential = PVehicle->GetDifferential();
+
+	float WheelRPM = 0;
+	for (int I = 0; I < PVehicle->Wheels.Num(); I++)
+	{
+		if (PVehicle->Wheels[I].Setup().EngineEnabled)
+		{
+			WheelRPM = FMath::Abs(PVehicle->Wheels[I].GetWheelRPM());
+		}
+	}
+
+	PEngine.SetEngineRPM(PTransmission.IsOutOfGear(), PTransmission.GetEngineRPMFromWheelRPM(WheelRPM));
+	PEngine.Simulate(DeltaTime);
+
+	PTransmission.SetEngineRPM(PEngine.GetEngineRPM()); // needs engine RPM to decide when to change gear (automatic gearbox)
+	PTransmission.SetAllowedToChangeGear(!VehicleState.bVehicleInAir && !IsWheelSpinning());
+	float GearRatio = PTransmission.GetGearRatio(PTransmission.GetCurrentGear());
+
+	PTransmission.Simulate(DeltaTime);
+
+	float TransmissionTorque = PTransmission.GetTransmissionTorque(PEngine.GetEngineTorque());
+
+	// apply drive torque to wheels
+	for (int WheelIdx = 0; WheelIdx < Wheels.Num(); WheelIdx++)
+	{
+		auto& PWheel = PVehicle->Wheels[WheelIdx];
+		if (PWheel.Setup().EngineEnabled)
+		{
+			check(PVehicle->NumDrivenWheels > 0);
+
+			if (PDifferential.Setup().DifferentialType == EDifferentialType::AllWheelDrive)
+			{
+				float SplitTorque = 1.0f;
+
+				if (PWheel.Setup().AxleType == FSimpleWheelConfig::EAxleType::Front)
+				{
+					SplitTorque = (1.0f - PDifferential.Setup().FrontRearSplit);
+				}
+				else
+				{
+					SplitTorque = PDifferential.Setup().FrontRearSplit;
+				}
+
+				PWheel.SetDriveTorque(MToCm(TransmissionTorque * SplitTorque) / (float)PVehicle->NumDrivenWheels);
+			}
+			else
+			{
+				PWheel.SetDriveTorque(MToCm(TransmissionTorque) / (float)PVehicle->NumDrivenWheels);
+			}
+		}
+	}
+}
+
+void UChaosWheeledVehicleSimulation::DrawDebug3D()
+{
+#if !(UE_BUILD_SHIPPING || UE_BUILD_TEST)
+
+	UChaosVehicleSimulation::DrawDebug3D();
+
+	if (PVehicle == nullptr)
+	{
+		return;
+	}
+
+	const FTransform BodyTransform = VehicleState.VehicleWorldTransform;
+
+	if (GWheeledVehicleDebugParams.ShowSuspensionLimits)
+	{
+		for (int WheelIdx = 0; WheelIdx < PVehicle->Suspension.Num(); WheelIdx++)
+		{
+			auto& PSuspension = PVehicle->Suspension[WheelIdx];
+			auto& PWheel = PVehicle->Wheels[WheelIdx];
+			// push the visualization out a bit sideways from the wheel model so we can actually see it
+			FVector VehicleRightAxis = VehicleState.VehicleWorldTransform.GetUnitAxis(EAxis::Y) * 48.0f;
+			const FVector& WheelOffset = PSuspension.GetLocalRestingPosition();
+			if (WheelOffset.Y < 0.0f)
+			{
+				VehicleRightAxis = VehicleRightAxis * -1.0f;
+			}
+
+			FVector LocalDirection = PSuspension.Setup().SuspensionAxis;
+			FVector WorldLocation = BodyTransform.TransformPosition(WheelOffset);
+			FVector WorldDirection = BodyTransform.TransformVector(LocalDirection);
+
+			FVector Start = WorldLocation + WorldDirection * (PWheel.GetEffectiveRadius() - PSuspension.Setup().SuspensionMaxRaise);
+			FVector End = WorldLocation + WorldDirection * (PWheel.GetEffectiveRadius() + PSuspension.Setup().SuspensionMaxDrop);
+
+			DrawDebugLine(World, Start + VehicleRightAxis, End + VehicleRightAxis, FColor::Orange, false, -1.f, 0, 3.f);
+
+			FVector Start2 = WorldLocation - WorldDirection * PSuspension.Setup().SuspensionMaxRaise;
+			FVector End2 = WorldLocation + WorldDirection * PSuspension.Setup().SuspensionMaxDrop;
+
+			DrawDebugLine(World, Start2 + VehicleRightAxis, End2 + VehicleRightAxis, FColor::Yellow, false, -1.f, 0, 3.f);
+		}
+	}
+
+	if (GWheeledVehicleDebugParams.ShowWheelCollisionNormal)
+	{
+		for (int WheelIdx = 0; WheelIdx < Wheels.Num(); WheelIdx++)
+		{
+			FHitResult& Hit = Wheels[WheelIdx]->HitResult;
+			DrawDebugLine(World, Hit.ImpactPoint, Hit.ImpactPoint + Hit.Normal * 20.0f, FColor::Yellow, true, 1.0f, 0, 1.0f);
+		}
+	}
+
+	if (GWheeledVehicleDebugParams.ShowSuspensionRaycasts)
+	{
+		for (int WheelIdx = 0; WheelIdx < Wheels.Num(); WheelIdx++)
+		{
+			const FVector& TraceStart = WheelState.Trace[WheelIdx].Start;
+			const FVector& TraceEnd = WheelState.Trace[WheelIdx].End;
+
+			// push the visualization out a bit sideways from the wheel model so we can actually see it
+			FVector VehicleRightAxis = VehicleState.VehicleWorldTransform.GetUnitAxis(EAxis::Y) * 50.0f;
+			const FVector& WheelOffset = PVehicle->Suspension[WheelIdx].GetLocalRestingPosition();
+			if (WheelOffset.Y < 0.0f)
+			{
+				VehicleRightAxis = VehicleRightAxis * -1.0f;
+			}
+
+			FColor UseColor = PVehicle->Wheels[WheelIdx].InContact() ? FColor::Green : FColor::Red;
+			DrawDebugDirectionalArrow(World, TraceStart + VehicleRightAxis, TraceEnd + VehicleRightAxis, 10.f, UseColor, false, -1.f, 0, 2.f);
+
+			DrawDebugLine(World, TraceStart, TraceStart + VehicleRightAxis, FColor::White, false, -1.f, 0, 1.f);
+			DrawDebugLine(World, TraceEnd, TraceEnd + VehicleRightAxis, FColor::White, false, -1.f, 0, 1.f);
+		}
+	}
+#endif
+}
+
+void UChaosWheeledVehicleSimulation::FillOutputState(FChaosVehicleAsyncOutput& Output)
+{
+	// #Note: remember to copy/interpolate values from the physics thread output in UChaosVehicleMovementComponent::ParallelUpdate
+	const auto& VehicleWheels = PVehicle->Wheels;
+	auto& VehicleSuspension = PVehicle->Suspension;
+	if (PVehicle->HasTransmission())
+	{
+		auto& Transmission = PVehicle->GetTransmission();
+		Output.VehicleSimOutput.CurrentGear = Transmission.GetCurrentGear();
+		Output.VehicleSimOutput.TargetGear = Transmission.GetTargetGear();
+		Output.VehicleSimOutput.TransmissionRPM = Transmission.GetTransmissionRPM();
+		Output.VehicleSimOutput.TransmissionTorque = Transmission.GetTransmissionTorque(PVehicle->GetEngine().GetTorqueFromRPM(false));
+	}
+	if (PVehicle->HasEngine())
+	{
+		auto& Engine = PVehicle->GetEngine();
+		Output.VehicleSimOutput.EngineRPM = Engine.GetEngineRPM();
+		Output.VehicleSimOutput.EngineTorque = Engine.GetEngineTorque();
+	}
+
+	// #TODO: can we avoid copies when async is turned off
+	for (int WheelIdx = 0; WheelIdx < VehicleWheels.Num(); WheelIdx++)
+	{
+		FWheelsOutput WheelsOut;
+		WheelsOut.InContact = VehicleWheels[WheelIdx].InContact();
+		WheelsOut.SteeringAngle = VehicleWheels[WheelIdx].GetSteeringAngle();
+		WheelsOut.AngularPosition = VehicleWheels[WheelIdx].GetAngularPosition();
+
+		WheelsOut.LateralAdhesiveLimit = VehicleWheels[WheelIdx].LateralAdhesiveLimit;
+		WheelsOut.LongitudinalAdhesiveLimit = VehicleWheels[WheelIdx].LongitudinalAdhesiveLimit;
+
+		WheelsOut.bIsSlipping = VehicleWheels[WheelIdx].IsSlipping();
+		WheelsOut.SlipMagnitude = VehicleWheels[WheelIdx].GetSlipMagnitude();
+		WheelsOut.bIsSkidding = VehicleWheels[WheelIdx].IsSkidding();
+		WheelsOut.SkidMagnitude = VehicleWheels[WheelIdx].GetSkidMagnitude();
+		WheelsOut.SkidNormal = WheelState.WorldWheelVelocity[WheelIdx].GetSafeNormal();
+
+		WheelsOut.SuspensionOffset = VehicleSuspension[WheelIdx].GetSuspensionOffset();
+		WheelsOut.SpringForce = VehicleSuspension[WheelIdx].GetSuspensionForce();
+		WheelsOut.NormalizedSuspensionLength = VehicleSuspension[WheelIdx].GetNormalizedLength();
+		Output.VehicleSimOutput.Wheels.Add(WheelsOut);
+	}
+
+}
+
+/**
+ * UChaosWheeledVehicleMovementComponent
+ */
 UChaosWheeledVehicleMovementComponent::UChaosWheeledVehicleMovementComponent(const FObjectInitializer& ObjectInitializer) : Super(ObjectInitializer)
 {
 	// default values setup
@@ -100,9 +879,6 @@
 	bMechanicalSimEnabled = true;
 	bSuspensionEnabled = true;
 	bWheelFrictionEnabled = true;
-
-	NumDrivenWheels = 0;
-
 }
 
 // Public
@@ -177,17 +953,29 @@
 
 									if (ConstraintHandle.IsValid())
 									{
-										const Chaos::FSimpleSuspensionConfig& SusSettings = PVehicle->GetSuspension(WheelIdx).Setup();
+										UChaosVehicleWheel* Wheel = Wheels[WheelIdx];
+										check(Wheel);
+										if (Chaos::FSuspensionConstraint* Constraint = static_cast<Chaos::FSuspensionConstraint*>(ConstraintHandle.Constraint))
+										{
+											Constraint->SetHardstopStiffness(1.0f);
+											Constraint->SetSpringStiffness(Wheel->SpringRate * 0.25f);
+											Constraint->SetSpringPreload(Wheel->SpringPreload);
+											Constraint->SetSpringDamping(Wheel->SuspensionDampingRatio * 5.0f);
+											Constraint->SetMinLength(-Wheel->SuspensionMaxRaise);
+											Constraint->SetMaxLength(Wheel->SuspensionMaxDrop);
+											Constraint->SetAxis(-Wheel->SuspensionAxis);
+										}
+
 										ConstraintHandles.Add(ConstraintHandle);
 										if (Chaos::FSuspensionConstraint* Constraint = static_cast<Chaos::FSuspensionConstraint*>(ConstraintHandle.Constraint))
 										{
 											Constraint->SetHardstopStiffness(1.0f);
-											Constraint->SetSpringStiffness(SusSettings.SpringRate * 0.25f);
-											Constraint->SetSpringPreload(SusSettings.SpringPreload);
-											Constraint->SetSpringDamping(SusSettings.DampingRatio * 5.0f);
-											Constraint->SetMinLength(-SusSettings.SuspensionMaxRaise);
-											Constraint->SetMaxLength(SusSettings.SuspensionMaxDrop);
-											Constraint->SetAxis(-SusSettings.SuspensionAxis);
+											Constraint->SetSpringStiffness(Wheel->SpringRate * 0.25f);
+											Constraint->SetSpringPreload(Wheel->SpringPreload);
+											Constraint->SetSpringDamping(Wheel->SuspensionDampingRatio * 5.0f);
+											Constraint->SetMinLength(-Wheel->SuspensionMaxRaise);
+											Constraint->SetMaxLength(Wheel->SuspensionMaxDrop);
+											Constraint->SetAxis(-Wheel->SuspensionAxis);
 										}
 									}
 #endif // WITH_CHAOS
@@ -234,23 +1022,17 @@
 	return true;
 }
 
-bool UChaosWheeledVehicleMovementComponent::CanSimulate() const
-{
-	if (Super::CanSimulate() == false)
-	{
-		return false;
-	}
 
-	return (PVehicle && PVehicle.IsValid()
-		&& PVehicle->Engine.Num() == PVehicle->Transmission.Num()
-		&& Wheels.Num() > 0 && Wheels.Num() == PVehicle->Suspension.Num());
-}
-
 void UChaosWheeledVehicleMovementComponent::OnCreatePhysicsState()
 {
 	Super::OnCreatePhysicsState();
+}
 
-	if (PVehicle)
+void UChaosWheeledVehicleMovementComponent::CreateVehicle()
+{
+	Super::CreateVehicle();
+
+	if (PVehicleOutput)
 	{
 		CreateWheels();
 
@@ -264,11 +1046,12 @@
 			}
 		}
 	}
+
 }
 
 void UChaosWheeledVehicleMovementComponent::OnDestroyPhysicsState()
 {
-	if (PVehicle.IsValid())
+	if (PVehicleOutput.IsValid())
 	{
 		if (MeshOnPhysicsStateChangeHandle.IsValid())
 		{
@@ -297,18 +1080,7 @@
 	
 }
 
-void UChaosWheeledVehicleMovementComponent::TickVehicle(float DeltaTime)
-{
-	Super::TickVehicle(DeltaTime);
 
-	// update wheels
-	for (int32 i = 0; i < Wheels.Num(); i++)
-	{
-		UChaosVehicleWheel* VehicleWheel = Wheels[i];
-		Wheels[i]->Tick(DeltaTime);
-	}
-}
-
 void UChaosWheeledVehicleMovementComponent::NextDebugPage()
 {
 	int PageAsInt = (int)DebugPage;
@@ -373,14 +1145,13 @@
 	Wheels.Reset();
 }
 
-void UChaosWheeledVehicleMovementComponent::SetupVehicle()
+void UChaosWheeledVehicleMovementComponent::SetupVehicle(TUniquePtr<Chaos::FSimpleWheeledVehicle>& PVehicle)
 {
 	check(PVehicle);
 
-	Super::SetupVehicle();
+	Super::SetupVehicle(PVehicle);
 
 	// we are allowed any number of wheels not limited to only 4
-	NumDrivenWheels = 0;
 	for (int32 WheelIdx = 0; WheelIdx < WheelSetups.Num(); ++WheelIdx)
 	{
 		UChaosVehicleWheel* Wheel = WheelSetups[WheelIdx].WheelClass.GetDefaultObject();
@@ -414,6 +1185,9 @@
 		WheelSim.SetWheelRadius(Wheel->WheelRadius); // initial radius
 		PVehicle->Wheels.Add(WheelSim);
 
+		FWheelsOutput WheelsOutput; // Receptacle for Data coming out of physics simulation on physics thread
+		PVehicleOutput->Wheels.Add(WheelsOutput);
+
 		Chaos::FSimpleSuspensionSim SuspensionSim(&Wheel->GetPhysicsSuspensionConfig());
 		PVehicle->Suspension.Add(SuspensionSim);
 
@@ -425,25 +1199,59 @@
 		// for debugging to identify a single wheel
 		PVehicle->Wheels[WheelIdx].SetWheelIndex(WheelIdx);
 		PVehicle->Suspension[WheelIdx].SetSpringIndex(WheelIdx);
+		PVehicle->NumDrivenWheels = NumDrivenWheels;
 	}
 
+	RecalculateAxles();
+
+	// setup axles in PVehicle
+	for (auto& Axle : AxleToWheelMap)
+	{
+		TArray<int>& WheelIndices = Axle.Value;
+
+		FChaosWheelSetup& WheelSetup = WheelSetups[WheelIndices[0]];
+		UChaosVehicleWheel* WheelData = WheelSetup.WheelClass.GetDefaultObject();
+
+		Chaos::FAxleSim AxleSim;
+		AxleSim.Setup.RollbarScaling = WheelData->RollbarScaling;
+
+		for (int WheelIdx : WheelIndices)
+		{
+			AxleSim.Setup.WheelIndex.Add(WheelIdx);
+		}
+
+		PVehicle->Axles.Add(AxleSim);
+
+	}
+
 	// cache this value as it's useful for steering setup calculations and debug rendering
 	WheelTrackDimensions = CalculateWheelLayoutDimensions();
 
 	if (bMechanicalSimEnabled)
+	if (bMechanicalSimEnabled)
 	{
 		Chaos::FSimpleEngineSim EngineSim(&EngineSetup.GetPhysicsEngineConfig());
 		PVehicle->Engine.Add(EngineSim);
 
 		Chaos::FSimpleTransmissionSim TransmissionSim(&TransmissionSetup.GetPhysicsTransmissionConfig());
 		PVehicle->Transmission.Add(TransmissionSim);
+
+		Chaos::FSimpleDifferentialSim DifferentialSim(&DifferentialSetup.GetPhysicsDifferentialConfig());
+		PVehicle->Differential.Add(DifferentialSim);
 	}
 
 	Chaos::FSimpleSteeringSim SteeringSim(&SteeringSetup.GetPhysicsSteeringConfig(WheelTrackDimensions));
 	PVehicle->Steering.Add(SteeringSim);
 
-	WheelState.Init(PVehicle->Wheels.Num());
+	Chaos::FTorqueControlSim TorqueSim(&TorqueControl.GetTorqueControlConfig());
+	PVehicle->TorqueControlSim.Add(TorqueSim);
 
+	Chaos::FTargetRotationControlSim TargetRotationSim(&TargetRotationControl.GetTargetRotationControlConfig());
+	PVehicle->TargetRotationControlSim.Add(TargetRotationSim);
+
+	Chaos::FStabilizeControlSim StabilizeSim(&StabilizeControl.GetStabilizeControlConfig());
+	PVehicle->StabilizeControlSim.Add(StabilizeSim);
+
 	// Setup the chassis and wheel shapes
 	SetupVehicleShapes();
 
@@ -451,8 +1259,7 @@
 	SetupVehicleMass();
 
 	// Setup Suspension
-	SetupSuspension();
-
+	SetupSuspension(PVehicle);
 }
 
 void UChaosWheeledVehicleMovementComponent::SetupVehicleShapes()
@@ -464,7 +1271,7 @@
 
 }
 
-void UChaosWheeledVehicleMovementComponent::SetupSuspension()
+void UChaosWheeledVehicleMovementComponent::SetupSuspension(TUniquePtr<Chaos::FSimpleWheeledVehicle>& PVehicle)
 {
 	if (!PVehicle.IsValid())
 	{
@@ -518,7 +1325,7 @@
 {
 	AxleToWheelMap.Empty();
 
-	for (int WheelIdx = 0; WheelIdx < Wheels.Num(); WheelIdx++)
+	for (int WheelIdx = 0; WheelIdx < WheelSetups.Num(); WheelIdx++)
 	{
 		FChaosWheelSetup& WheelSetup = WheelSetups[WheelIdx];
 		UChaosVehicleWheel* Wheel = WheelSetup.WheelClass.GetDefaultObject();
@@ -533,6 +1340,7 @@
 			WheelIndices.Add(WheelIdx);
 			AxleToWheelMap.Add(Wheel, WheelIndices);
 		}
+
 	}
 }
 
@@ -542,627 +1350,14 @@
 	return LocateBoneOffset(WheelSetup.BoneName, Offset);
 }
 
-// Update
-
-void UChaosWheeledVehicleMovementComponent::UpdateSimulation(float DeltaTime)
-{
-	// Inherit common vehicle simulation stages ApplyAerodynamics, ApplyTorqueControl, etc
-	Super::UpdateSimulation(DeltaTime);
-
-	FBodyInstance* TargetInstance = GetBodyInstance();
-
-	if (CanSimulate() && TargetInstance)
-	{
-		// sanity check that everything is setup ok
-		ensure(Wheels.Num() == PVehicle->Suspension.Num());
-		ensure(Wheels.Num() == PVehicle->Wheels.Num());
-		ensure(WheelState.LocalWheelVelocity.Num() == Wheels.Num());
-		ensure(WheelState.WheelWorldLocation.Num() == Wheels.Num());
-		ensure(WheelState.WorldWheelVelocity.Num() == Wheels.Num());
-
-		///////////////////////////////////////////////////////////////////////
-		// Cache useful state so we are not re-calculating the same data
-
-		for (int WheelIdx = 0; WheelIdx < Wheels.Num(); WheelIdx++)
-		{
-			WheelState.CaptureState(WheelIdx, PVehicle->Suspension[WheelIdx].GetLocalRestingPosition(), TargetInstance);
-		}
-
-		///////////////////////////////////////////////////////////////////////
-		// Suspension Raycast
-
-		for (int WheelIdx = 0; WheelIdx < Wheels.Num(); WheelIdx++)
-		{
-			auto& PSuspension = PVehicle->Suspension[WheelIdx];
-			auto& PWheel = PVehicle->Wheels[WheelIdx];
-			PSuspension.UpdateWorldRaycastLocation(VehicleState.VehicleWorldTransform, PWheel.Setup().WheelRadius, WheelState.Trace[WheelIdx]);
-		}
-
-		if (!GWheeledVehicleDebugParams.DisableSuspensionForces && bSuspensionEnabled)
-		{
-			PerformSuspensionTraces(WheelState.Trace);
-		}
-
-		//////////////////////////////////////////////////////////////////////////
-		// Wheel and Vehicle in air state
-
-		VehicleState.bVehicleInAir = true;
-		VehicleState.NumWheelsOnGround = 0;
-		for (int WheelIdx = 0; WheelIdx < Wheels.Num(); WheelIdx++)
-		{
-			// tell systems who care that wheel is touching the ground
-			PVehicle->Wheels[WheelIdx].SetOnGround(Wheels[WheelIdx]->HitResult.bBlockingHit);
-
-			// only requires one wheel to be on the ground for the vehicle to be NOT in the air
-			if (PVehicle->Wheels[WheelIdx].InContact())
-			{
-				VehicleState.bVehicleInAir = false;
-				VehicleState.NumWheelsOnGround++;
-			}
-		}
-		VehicleState.bAllWheelsOnGround = (VehicleState.NumWheelsOnGround == Wheels.Num());
-
-		///////////////////////////////////////////////////////////////////////
-		// Input
-		ApplyInput(DeltaTime);
-
-		///////////////////////////////////////////////////////////////////////
-		// Engine/Transmission
-		if (!GWheeledVehicleDebugParams.DisableSuspensionForces && bMechanicalSimEnabled)
-		{
-			ProcessMechanicalSimulation(DeltaTime);
-		}
-
-		///////////////////////////////////////////////////////////////////////
-		// Suspension
-
-		if (!GWheeledVehicleDebugParams.DisableSuspensionForces && bSuspensionEnabled)
-		{
-			ApplySuspensionForces(DeltaTime);
-		}
-
-		///////////////////////////////////////////////////////////////////////
-		// Steering
-
-		ProcessSteering();
-
-		///////////////////////////////////////////////////////////////////////
-		// Wheel Friction
-
-		if (!GWheeledVehicleDebugParams.DisableFrictionForces && bWheelFrictionEnabled)
-		{		
-			ApplyWheelFrictionForces(DeltaTime);
-		}
-
-		if (PerformanceMeasure.IsEnabled())
-		{
-			PerformanceMeasure.Update(DeltaTime, VehicleState.VehicleWorldTransform.GetLocation(), VehicleState.ForwardSpeed);
-		}
-
-		///////////////////////////////////////////////////////////////////////
-		// Consolidate wheel output state
-		FillWheelOutputState();
-
-	}
-
-}
-
-void UChaosWheeledVehicleMovementComponent::PerformSuspensionTraces(const TArray<FSuspensionTrace>& SuspensionTrace)
-{
-	//#todo: SpringCollisionChannel should be a parameter setup
-	ECollisionChannel SpringCollisionChannel = ECollisionChannel::ECC_WorldDynamic;
-
-	TArray<AActor*> ActorsToIgnore;
-	ActorsToIgnore.Add(GetPawnOwner()); // ignore self in scene query
-
-	FCollisionQueryParams TraceParams(NAME_None, FCollisionQueryParams::GetUnknownStatId(), false, nullptr);
-	TraceParams.bReturnPhysicalMaterial = true;	// we need this to get the surface friction coefficient
-	TraceParams.AddIgnoredActors(ActorsToIgnore);
-
-	FCollisionResponseParams ResponseParams;
-
-	// batching is about 0.5ms (25%) faster when there's 100 vehicles on a flat terrain
-	if (GVehicleDebugParams.BatchQueries)
-	{
-		FBox QueryBox;
-		for (int WheelIdx = 0; WheelIdx < Wheels.Num(); WheelIdx++)
-		{
-
-			FVector TraceStart = SuspensionTrace[WheelIdx].Start;
-			FVector TraceEnd = SuspensionTrace[WheelIdx].End;
-
-			if (WheelIdx == 0)
-			{
-				QueryBox = FBox(TraceStart, TraceEnd);
-			}
-			else
-			{
-				QueryBox.Min = QueryBox.Min.ComponentMin(TraceStart);
-				QueryBox.Min = QueryBox.Min.ComponentMin(TraceEnd);
-				QueryBox.Max = QueryBox.Max.ComponentMax(TraceStart);
-				QueryBox.Max = QueryBox.Max.ComponentMax(TraceEnd);
-			}
-		}
-		float OneWheelRadius = PVehicle->Wheels[0].GetEffectiveRadius(); // or wheel width
-
-		QueryBox.ExpandBy(FVector(OneWheelRadius, OneWheelRadius, OneWheelRadius)); // little extra just to be on the safe side consider 1 or 2 wheel vehicle
-		TArray<FOverlapResult> OverlapResults;
-		FCollisionShape CollisionBox;
-		CollisionBox.SetBox(QueryBox.GetExtent());
-
-#if !(UE_BUILD_SHIPPING || UE_BUILD_TEST)
-		if (GWheeledVehicleDebugParams.ShowBatchQueryExtents)
-		{
-			DrawDebugBox(GetWorld(), QueryBox.GetCenter(), QueryBox.GetExtent(), FColor::Yellow, false, -1.0f, 0, 2.0f);
-		}
-#endif
-
-		const bool bOverlapHit = GetWorld()->OverlapMultiByChannel(OverlapResults, QueryBox.GetCenter(), FQuat::Identity, SpringCollisionChannel, CollisionBox, TraceParams, ResponseParams);
-
-		for (int32 WheelIdx = 0; WheelIdx < Wheels.Num(); ++WheelIdx)
-		{
-			FHitResult& HitResult = Wheels[WheelIdx]->HitResult;
-			HitResult = FHitResult();
-
-			if (bOverlapHit)
-			{
-				const FVector& TraceStart = SuspensionTrace[WheelIdx].Start;
-				const FVector& TraceEnd = SuspensionTrace[WheelIdx].End;
-				TraceParams.bTraceComplex = (Wheels[WheelIdx]->SweepType == ESweepType::ComplexSweep);
-
-				// Test each overlapped object for a hit result
-				for (FOverlapResult OverlapResult : OverlapResults)
-				{
-					if (!OverlapResult.bBlockingHit)
-						continue;
-
-					FHitResult ComponentHit;
-
-					switch (Wheels[WheelIdx]->SweepShape)
-					{
-						case ESweepShape::Spherecast:
-						{
-							float WheelRadius = PVehicle->Wheels[WheelIdx].GetEffectiveRadius(); // or wheel width
-							FVector VehicleUpAxis = GetOwner()->GetTransform().GetUnitAxis(EAxis::Z);
-
-							FVector Start = TraceStart + VehicleUpAxis * WheelRadius;
-							FVector End = TraceEnd + VehicleUpAxis * WheelRadius;
-
-							if (OverlapResult.Component->SweepComponent(ComponentHit, Start, End, FQuat::Identity, FCollisionShape::MakeSphere(WheelRadius), TraceParams.bTraceComplex))
-							{
-								if (ComponentHit.Time < HitResult.Time)
-								{
-									HitResult = ComponentHit;
-									HitResult.bBlockingHit = OverlapResult.bBlockingHit;
-								}
-							}
-						}
-						break;
-
-						case ESweepShape::Raycast:
-						default:
-						{
-							if (OverlapResult.Component->LineTraceComponent(ComponentHit, TraceStart, TraceEnd, TraceParams))
-							{
-								if (ComponentHit.Time < HitResult.Time)
-								{
-									HitResult = ComponentHit;
-									HitResult.bBlockingHit = OverlapResult.bBlockingHit;
-								}
-							}
-						}
-						break;
-					}
-				}
-			}
-		}
-	}
-	else
-	{
-		for (int WheelIdx = 0; WheelIdx < Wheels.Num(); WheelIdx++)
-		{
-			FHitResult& HitResult = Wheels[WheelIdx]->HitResult;
-
-			FVector TraceStart = SuspensionTrace[WheelIdx].Start;
-			FVector TraceEnd = SuspensionTrace[WheelIdx].End;
-			TraceParams.bTraceComplex = (Wheels[WheelIdx]->SweepType == ESweepType::ComplexSweep);
-
-			switch (Wheels[WheelIdx]->SweepShape)
-			{
-				case ESweepShape::Spherecast:
-				{
-					//float Radius = PVehicle->Wheels[WheelIdx].GetEffectiveRadius(); // or wheel width
-					float Radius = PVehicle->Wheels[WheelIdx].Setup().WheelWidth * 0.5f; // or wheel width
-					FVector VehicleUpAxis = GetOwner()->GetTransform().GetUnitAxis(EAxis::Z);
-
-					GetWorld()->SweepSingleByChannel(HitResult
-						, TraceStart + VehicleUpAxis * Radius
-						, TraceEnd + VehicleUpAxis * Radius
-						, FQuat::Identity, SpringCollisionChannel
-						, FCollisionShape::MakeSphere(Radius), TraceParams
-						, FCollisionResponseParams::DefaultResponseParam);
-				}
-				break;
-
-
-
-				case ESweepShape::Raycast:
-				default:
-				{
-					GetWorld()->LineTraceSingleByChannel(HitResult, TraceStart, TraceEnd, SpringCollisionChannel, TraceParams, FCollisionResponseParams::DefaultResponseParam);
-				}
-				break;
-			}
-		}
-	}
-
-}
-
-void UChaosWheeledVehicleMovementComponent::ApplyWheelFrictionForces(float DeltaTime)
-{
-	FBodyInstance* TargetInstance = GetBodyInstance();
-
-	for (int WheelIdx = 0; WheelIdx < Wheels.Num(); WheelIdx++)
-	{
-		auto& PWheel = PVehicle->Wheels[WheelIdx]; // Physics Wheel
-		FHitResult& HitResult = Wheels[WheelIdx]->HitResult;
-		UChaosVehicleWheel* VehicleWheel = Wheels[WheelIdx];
-
-		if (PWheel.InContact())
-		{
-			if (HitResult.PhysMaterial.IsValid())
-			{
-				PWheel.SetSurfaceFriction(HitResult.PhysMaterial->Friction);
-			}
-
-			// take into account steering angle
-			float SteerAngleDegrees = VehicleWheel->GetSteerAngle(); // temp
-			FRotator SteeringRotator(0.f, SteerAngleDegrees, 0.f);
-			FVector SteerLocalWheelVelocity = SteeringRotator.UnrotateVector(WheelState.LocalWheelVelocity[WheelIdx]);
-
-			PWheel.SetVehicleGroundSpeed(SteerLocalWheelVelocity);
-			PWheel.Simulate(DeltaTime);
-
-			float RotationAngle = FMath::RadiansToDegrees(PWheel.GetAngularPosition());
-			FVector FrictionForceLocal = PWheel.GetForceFromFriction();
-			FrictionForceLocal = SteeringRotator.RotateVector(FrictionForceLocal);
-
-			FVector GroundZVector = HitResult.Normal;
-			FVector GroundXVector = FVector::CrossProduct(VehicleState.VehicleRightAxis, GroundZVector);
-			FVector GroundYVector = FVector::CrossProduct(GroundZVector, GroundXVector);
-
-			// the force should be applied along the ground surface not along vehicle forward vector?
-			//FVector FrictionForceVector = VehicleState.VehicleWorldTransform.TransformVector(FrictionForceLocal);
-			FMatrix Mat(GroundXVector, GroundYVector, GroundZVector, VehicleState.VehicleWorldTransform.GetLocation());
-			FVector FrictionForceVector = Mat.TransformVector(FrictionForceLocal);
-
-			check(PWheel.InContact());
-			AddForceAtPosition(FrictionForceVector, WheelState.WheelWorldLocation[WheelIdx]);
-
-#if !(UE_BUILD_SHIPPING || UE_BUILD_TEST)
-			if (GWheeledVehicleDebugParams.ShowWheelForces)
-			{
-				// show longitudinal drive force
-				{
-					DrawDebugLine(GetWorld()
-						, WheelState.WheelWorldLocation[WheelIdx]
-						, WheelState.WheelWorldLocation[WheelIdx] + FrictionForceVector * 0.001f
-						, FColor::Yellow, false, -1.0f, 0, 2);
-
-					DrawDebugLine(GetWorld()
-						, WheelState.WheelWorldLocation[WheelIdx]
-						, WheelState.WheelWorldLocation[WheelIdx] + GroundZVector * 100.f
-						, FColor::Orange, false, -1.0f, 0, 2);
-
-				}
-
-			}
-#endif
-
-		}
-		else
-		{
-			PWheel.SetVehicleGroundSpeed(WheelState.LocalWheelVelocity[WheelIdx]);
-			PWheel.Simulate(DeltaTime);
-		}
-
-	}
-}
-
-void UChaosWheeledVehicleMovementComponent::ApplySuspensionForces(float DeltaTime)
-{
-	TArray<float> SusForces;
-	SusForces.Init(0.f, Wheels.Num());
-	FBodyInstance* TargetInstance = GetBodyInstance();
-
-	for (int WheelIdx = 0; WheelIdx < Wheels.Num(); WheelIdx++)
-	{
-		FChaosWheelSetup& WheelSetup = WheelSetups[WheelIdx];
-		UChaosVehicleWheel* Wheel = WheelSetup.WheelClass.GetDefaultObject();
-		FHitResult& HitResult = Wheels[WheelIdx]->HitResult;
-
-		float NewDesiredLength = 1.0f; // suspension max length
-		float ForceMagnitude2 = 0.f;
-		auto& PWheel = PVehicle->Wheels[WheelIdx];
-		auto& PSuspension = PVehicle->Suspension[WheelIdx];
-		float SuspensionMovePosition = -PSuspension.Setup().MaxLength;
-		
-		if (!GWheeledVehicleDebugParams.DisableSuspensionConstraint)
-		{
-			FPhysicsCommand::ExecuteWrite(TargetInstance->ActorHandle, [&](const FPhysicsActorHandle& Chassis)
-			{
-#if WITH_CHAOS
-				if (ConstraintHandles.Num() > 0)
-				{
-					FPhysicsConstraintHandle& ConstraintHandle = ConstraintHandles[WheelIdx];
-					if (ConstraintHandle.IsValid())
-					{
-						if (Chaos::FSuspensionConstraint* Constraint = static_cast<Chaos::FSuspensionConstraint*>(ConstraintHandle.Constraint))
-						{
-							FVec3 P = HitResult.ImpactPoint + (PWheel.Setup().WheelRadius * VehicleState.VehicleUpAxis);
-							Constraint->SetTarget( P );
-							Constraint->SetEnabled(PWheel.InContact());
-						}
-					}
-
-				}
-#endif // WITH_CHAOS
-			});
-		}
-
-		if (PWheel.InContact())
-		{
-			NewDesiredLength = HitResult.Distance;
-
-			SuspensionMovePosition = -FVector::DotProduct(WheelState.WheelWorldLocation[WheelIdx] - HitResult.ImpactPoint, VehicleState.VehicleUpAxis) + Wheel->WheelRadius;
-
-			PSuspension.SetSuspensionLength(NewDesiredLength, PWheel.Setup().WheelRadius);
-			PSuspension.SetLocalVelocity(WheelState.LocalWheelVelocity[WheelIdx]);
-			PSuspension.Simulate(DeltaTime);
-
-			float ForceMagnitude = PSuspension.GetSuspensionForce();
-
-			FVector GroundZVector = HitResult.Normal;
-			FVector SuspensionForceVector = VehicleState.VehicleUpAxis * ForceMagnitude;
-
-			FVector SusApplicationPoint = WheelState.WheelWorldLocation[WheelIdx] + PVehicle->Suspension[WheelIdx].Setup().SuspensionForceOffset;
-
-			check(PWheel.InContact());
-			if (GWheeledVehicleDebugParams.DisableSuspensionConstraint)
-			{
-				AddForceAtPosition(SuspensionForceVector, SusApplicationPoint);
-			}
-
-			ForceMagnitude = PSuspension.Setup().WheelLoadRatio * ForceMagnitude + (1.f - PSuspension.Setup().WheelLoadRatio) * PSuspension.Setup().RestingForce;
-			PWheel.SetWheelLoadForce(ForceMagnitude);
-			PWheel.SetMassPerWheel(TargetInstance->GetBodyMass() / PVehicle->Wheels.Num());
-			SusForces[WheelIdx] = ForceMagnitude;
-
-#if !(UE_BUILD_SHIPPING || UE_BUILD_TEST)
-			if (GWheeledVehicleDebugParams.ShowSuspensionForces)
-			{
-				DrawDebugLine(GetWorld()
-					, SusApplicationPoint
-					, SusApplicationPoint + SuspensionForceVector * GVehicleDebugParams.ForceDebugScaling
-					, FColor::Blue, false, -1.0f, 0, 5);
-
-				DrawDebugLine(GetWorld()
-					, SusApplicationPoint
-					, SusApplicationPoint + GroundZVector * 140.f
-					, FColor::Yellow, false, -1.0f, 0, 5);
-			}
-#endif
-
-		}
-		else
-		{
-			PSuspension.SetSuspensionLength(PSuspension.GetTraceLength(PWheel.Setup().WheelRadius), PWheel.Setup().WheelRadius);
-		}
-
-	}
-
-
-	if (!GWheeledVehicleDebugParams.DisableRollbarForces && AxleToWheelMap.Num() > 0)
-	{
-		for (auto& Axle : AxleToWheelMap)
-		{
-			const UChaosVehicleWheel* Wheel = Axle.Key;
-			const TArray<int>& WheelIndices = Axle.Value;
-
-			//#todo: only works with 2 wheels on an axle at present
-			if (WheelIndices.Num() == 2)
-			{
-				int WheelIdxA = WheelIndices[0];
-				int WheelIdxB = WheelIndices[1];
-
-				float FV = Wheel->RollbarScaling;
-				float ForceDiffOnAxleF = SusForces[WheelIdxA] - SusForces[WheelIdxB];
-				FVector ForceVector0 = VehicleState.VehicleUpAxis * ForceDiffOnAxleF * FV;
-				FVector ForceVector1 = VehicleState.VehicleUpAxis * ForceDiffOnAxleF * -FV;
-
-				FVector SusApplicationPoint0 = WheelState.WheelWorldLocation[WheelIdxA] + PVehicle->Suspension[WheelIdxA].Setup().SuspensionForceOffset;
-				AddForceAtPosition(ForceVector0, SusApplicationPoint0);
-
-				FVector SusApplicationPoint1 = WheelState.WheelWorldLocation[WheelIdxB] + PVehicle->Suspension[WheelIdxB].Setup().SuspensionForceOffset;
-				AddForceAtPosition(ForceVector1, SusApplicationPoint1);
-			}
-		}
-	}
-}
-
-void UChaosWheeledVehicleMovementComponent::ProcessSteering()
-{
-	auto& PSteering = PVehicle->GetSteering();
-
-	for (int WheelIdx = 0; WheelIdx < Wheels.Num(); WheelIdx++)
-	{
-		auto& PWheel = PVehicle->Wheels[WheelIdx]; // Physics Wheel
-		FHitResult& HitResult = Wheels[WheelIdx]->HitResult;
-
-		if (PWheel.Setup().SteeringEnabled)
-		{
-			FRichCurve* SteeringCurveData = SteeringSetup.SteeringCurve.GetRichCurve();
-			float SpeedScale = 1.0f;
-
-			// allow full counter steering when steering into a power slide
-			//if (SteeringInput * VehicleState.VehicleLocalVelocity.Y > 0.0f)
-			{
-				SpeedScale = SteeringCurveData->Eval(CmSToMPH(VehicleState.ForwardSpeed));
-			}
-
-			float UseSteeringValue = SteeringInput * SpeedScale;
-
-			float SteeringAngle = 0.f;
-			if (FMath::Abs(GWheeledVehicleDebugParams.SteeringOverride) > 0.01f)
-			{
-				SteeringAngle = PWheel.Setup().MaxSteeringAngle * GWheeledVehicleDebugParams.SteeringOverride;
-			}
-			else
-			{
-				float WheelSide = PVehicle->GetSuspension(WheelIdx).GetLocalRestingPosition().Y;
-				SteeringAngle = PSteering.GetSteeringAngle(UseSteeringValue, PWheel.Setup().MaxSteeringAngle, WheelSide);
-			}
-
-			PWheel.SetSteeringAngle(SteeringAngle);
-		}
-		else
-		{
-			PWheel.SetSteeringAngle(0.0f);
-		}
-	}
-}
-
-void UChaosWheeledVehicleMovementComponent::ApplyInput(float DeltaTime)
-{
-	Super::ApplyInput(DeltaTime);
-
-	float EngineBraking = 0.f;
-	if (bMechanicalSimEnabled)
-	{
-		auto& PEngine = PVehicle->GetEngine();
-		auto& PTransmission = PVehicle->GetTransmission();
-
-		if (bRawGearUpInput)
-		{
-			PTransmission.ChangeUp();
-			bRawGearUpInput = false;
-		}
-
-		if (bRawGearDownInput)
-		{
-			PTransmission.ChangeDown();
-			bRawGearDownInput = false;
-		}
-
-		if (GWheeledVehicleDebugParams.ThrottleOverride > 0.f)
-		{
-			PTransmission.SetGear(1, true);
-			BrakeInput = 0.f;
-			PEngine.SetThrottle(GWheeledVehicleDebugParams.ThrottleOverride);
-		}
-		else
-		{
-			PEngine.SetThrottle(ThrottleInput* ThrottleInput);
-		}
-
-		EngineBraking = PEngine.GetEngineRPM() * PEngine.Setup().EngineBrakeEffect;
-	}
-
-	for (int WheelIdx = 0; WheelIdx < Wheels.Num(); WheelIdx++)
-	{
-		auto& PWheel = PVehicle->Wheels[WheelIdx];
-
-		float EngineBrakingForce = 0.0f;
-		if ((ThrottleInput < SMALL_NUMBER) && FMath::Abs(VehicleState.ForwardSpeed) > SMALL_NUMBER && PWheel.Setup().EngineEnabled)
-		{
-			EngineBrakingForce = EngineBraking;
-		}
-
-		if (PWheel.Setup().BrakeEnabled)
-		{
-			float BrakeForce = PWheel.Setup().MaxBrakeTorque * BrakeInput;
-			PWheel.SetBrakeTorque(MToCm(BrakeForce + EngineBrakingForce));
-		}
-		else
-		{
-			PWheel.SetBrakeTorque(MToCm(EngineBraking));
-		}
-
-		if (bRawHandbrakeInput && PWheel.Setup().HandbrakeEnabled)
-		{
-			PWheel.SetBrakeTorque(MToCm(bRawHandbrakeInput * PWheel.Setup().HandbrakeTorque));
-		}
-	}
-
-
-}
-
-void UChaosWheeledVehicleMovementComponent::ProcessMechanicalSimulation(float DeltaTime)
-{
-	auto& PEngine = PVehicle->GetEngine();
-	auto& PTransmission = PVehicle->GetTransmission();
-
-	float WheelRPM = 0;
-	for (int I=0; I<PVehicle->Wheels.Num(); I++)
-	{
-		if (PVehicle->Wheels[I].Setup().EngineEnabled)
-		{ 
-			WheelRPM = FMath::Abs(PVehicle->Wheels[I].GetWheelRPM());
-		}
-	}
-
-	PEngine.SetEngineRPM(PTransmission.IsOutOfGear() , PTransmission.GetEngineRPMFromWheelRPM(WheelRPM));
-	PEngine.Simulate(DeltaTime);
-
-	PTransmission.SetEngineRPM(PEngine.GetEngineRPM()); // needs engine RPM to decide when to change gear (automatic gearbox)
-	PTransmission.SetAllowedToChangeGear(!VehicleState.bVehicleInAir && !IsWheelSpinning());
-	float GearRatio = PTransmission.GetGearRatio(PTransmission.GetCurrentGear());
-
-	PTransmission.Simulate(DeltaTime);
-
-	float TransmissionTorque = PTransmission.GetTransmissionTorque(PEngine.GetEngineTorque());
-
-	// apply drive torque to wheels
-	for (int WheelIdx = 0; WheelIdx < Wheels.Num(); WheelIdx++)
-	{
-		UChaosVehicleWheel* Wheel = WheelSetups[WheelIdx].WheelClass.GetDefaultObject();
-		auto& PWheel = PVehicle->Wheels[WheelIdx];
-		if (PWheel.Setup().EngineEnabled)
-		{
-			check(NumDrivenWheels > 0);
-
-			if (DifferentialSetup.DifferentialType == EVehicleDifferential::AllWheelDrive)
-			{
-				float SplitTorque = 1.0f;
-
-				if (Wheel->GetAxleType() == EAxleType::Front)
-				{
-					SplitTorque = (1.0f - DifferentialSetup.FrontRearSplit);
-				}
-				else
-				{
-					SplitTorque = DifferentialSetup.FrontRearSplit;
-				}	
-
-				PWheel.SetDriveTorque(MToCm(TransmissionTorque * SplitTorque) / (float)NumDrivenWheels);
-			}
-			else
-			{
-				PWheel.SetDriveTorque(MToCm(TransmissionTorque) / (float)NumDrivenWheels);
-			}
-		}
-	}
-}
-
-
 // Access to data
 float UChaosWheeledVehicleMovementComponent::GetEngineRotationSpeed() const
 {
 	float EngineRPM = 0.f;
 
-	if (PVehicle.IsValid() && PVehicle->HasEngine())
+	if (bMechanicalSimEnabled)
 	{
-		EngineRPM = PVehicle->GetEngine().GetEngineRPM();
+		EngineRPM = PVehicleOutput->EngineRPM;
 	}
 
 	return EngineRPM;
@@ -1172,27 +1367,14 @@
 {
 	float MaxEngineRPM = 0.f;
 	
-	if (PVehicle.IsValid() && PVehicle->HasEngine())
+	if (bMechanicalSimEnabled)
 	{
-		MaxEngineRPM = PVehicle->GetEngine().Setup().MaxRPM;
+		MaxEngineRPM = EngineSetup.MaxRPM;
 	}
 
 	return MaxEngineRPM;
 }
 
-bool UChaosWheeledVehicleMovementComponent::IsWheelSpinning() const // rename? / skidding too, covered by GetWheelState
-{
-	for (auto& Wheel : PVehicle->Wheels)
-	{
-		if (Wheel.IsSlipping())
-		{
-			return true;
-		}
-	}
-
-	return false;
-}
-
 // Helper
 FVector2D UChaosWheeledVehicleMovementComponent::CalculateWheelLayoutDimensions()
 {
@@ -1226,15 +1408,16 @@
 // Debug
 void UChaosWheeledVehicleMovementComponent::DrawDebug(UCanvas* Canvas, float& YL, float& YPos)
 {
+	ensure(IsInGameThread());
+
 #if !(UE_BUILD_SHIPPING || UE_BUILD_TEST)
 
 	Super::DrawDebug(Canvas, YL, YPos);
 
-	FChaosVehicleManager* MyVehicleManager = FChaosVehicleManager::GetVehicleManagerFromScene(GetWorld()->GetPhysicsScene());
 	FBodyInstance* TargetInstance = GetBodyInstance();
 
 	// #todo: is this rendering multiple times in multiplayer
-	if (!PVehicle.IsValid() || TargetInstance == nullptr || MyVehicleManager == nullptr)
+	if (TargetInstance == nullptr)
 	{
 		return;
 	}
@@ -1242,12 +1425,8 @@
 	float ForwardSpeedMPH = CmSToMPH(GetForwardSpeed());
 
 	// always draw this even on (DebugPage == EDebugPages::BasicPage)
-	if (PVehicle->HasTransmission())
+	if (bMechanicalSimEnabled)
 	{
-		auto& PTransmission = PVehicle->GetTransmission();
-		auto& PEngine = PVehicle->GetEngine();
-		auto& PTransmissionSetup = PTransmission.Setup();
-
 		UFont* RenderFont = GEngine->GetLargeFont();
 		Canvas->SetDrawColor(FColor::Yellow);
 
@@ -1257,12 +1436,12 @@
 		float YLine = Y * 2.f - 50.f;
 		float Scaling = 2.f;
 		Canvas->DrawText(RenderFont, FString::Printf(TEXT("%d mph"), (int)ForwardSpeedMPH), X-100, YLine, Scaling, Scaling);
-		Canvas->DrawText(RenderFont, FString::Printf(TEXT("[%d]"), (int)PTransmission.GetCurrentGear()), X, YLine, Scaling, Scaling);
-		Canvas->DrawText(RenderFont, FString::Printf(TEXT("%d rpm"), (int)PEngine.GetEngineRPM()), X+50, YLine, Scaling, Scaling);
+		Canvas->DrawText(RenderFont, FString::Printf(TEXT("[%d]"), (int)PVehicleOutput->CurrentGear), X, YLine, Scaling, Scaling);
+		Canvas->DrawText(RenderFont, FString::Printf(TEXT("%d rpm"), (int)PVehicleOutput->EngineRPM), X+50, YLine, Scaling, Scaling);
 
 		FVector2D DialPos(X+10, YLine-40);
 		float DialRadius = 50;
-		DrawDial(Canvas, DialPos, DialRadius, PEngine.GetEngineRPM(), PEngine.Setup().MaxRPM);
+		DrawDial(Canvas, DialPos, DialRadius, PVehicleOutput->EngineRPM, EngineSetup.MaxRPM);
 
 	}
 
@@ -1272,28 +1451,31 @@
 		Canvas->SetDrawColor(FColor::White);
 		YPos += 16;
 		
-		if (PVehicle->HasTransmission())
+		if (bMechanicalSimEnabled)
 		{
-			YPos += Canvas->DrawText(RenderFont, FString::Printf(TEXT("RPM: %.1f (ChangeUp RPM %d, ChangeDown RPM %d)")
+			YPos += Canvas->DrawText(RenderFont, FString::Printf(TEXT("RPM: %.1f (ChangeUp RPM %.0f, ChangeDown RPM %.0f)")
 				, GetEngineRotationSpeed()
-				, PVehicle->GetTransmission().Setup().ChangeUpRPM
-				, PVehicle->GetTransmission().Setup().ChangeDownRPM), 4, YPos);
+				, TransmissionSetup.ChangeUpRPM
+				, TransmissionSetup.ChangeDownRPM), 4, YPos);
 
-			YPos += Canvas->DrawText(RenderFont, FString::Printf(TEXT("Gear: %d (Target %d)")
-				, GetCurrentGear(), GetTargetGear()), 4, YPos);
+			YPos += Canvas->DrawText(RenderFont, FString::Printf(TEXT("Gear: %d (Target %d)(Target %d)")
+				, GetCurrentGear(), TargetGearInput, PVehicleOutput->TargetGear), 4, YPos);
 		}
 		//YPos += Canvas->DrawText(RenderFont, FString::Printf(TEXT("Drag: %.1f"), DebugDragMagnitude), 4, YPos);
 
 		YPos += 16;
-		for (int i = 0; i < PVehicle->Wheels.Num(); i++)
+		for (int i = 0; i < PVehicleOutput->Wheels.Num(); i++)
 		{
-			YPos += Canvas->DrawText(RenderFont, FString::Printf(TEXT("WheelLoad: [%d] %1.f N"), i, CmToM(PVehicle->Wheels[i].GetWheelLoadForce())), 4, YPos);
+			YPos += Canvas->DrawText(RenderFont, FString::Printf(TEXT("WheelLoad: [%d] %1.f N"), i, CmToM(WheelStatus[i].SpringForce)), 4, YPos);
 		}
 
 		YPos += 16;
-		for (int i = 0; i < PVehicle->Wheels.Num(); i++)
-		{
-			YPos += Canvas->DrawText(RenderFont, FString::Printf(TEXT("SurfaceFriction: [%d] %.2f"), i, PVehicle->Wheels[i].GetSurfaceFriction()), 4, YPos);
+		for (int i = 0; i < PVehicleOutput->Wheels.Num(); i++)
+		{			
+			if (WheelStatus[i].PhysMaterial.IsValid())
+			{
+				YPos += Canvas->DrawText(RenderFont, FString::Printf(TEXT("SurfaceFriction: [%d] %.2f"), i, WheelStatus[i].PhysMaterial->Friction), 4, YPos);
+			}
 		}
 		
 	}
@@ -1326,88 +1508,88 @@
 		for (int32 WheelIdx = 0; WheelIdx < WheelSetups.Num(); ++WheelIdx)
 		{
 
-			auto& PWheel = PVehicle->Wheels[WheelIdx];
-			FVector Forces = PWheel.GetForceFromFriction();
-
-			FChaosWheelSetup& WheelSetup = WheelSetups[WheelIdx];
-			UChaosVehicleWheel* Wheel = WheelSetup.WheelClass.GetDefaultObject();
-			check(Wheel);
-			UPhysicalMaterial* ContactMat = Wheel->GetContactSurfaceMaterial();
-
-			const FVector WheelOffset = GetWheelRestingPosition(WheelSetup);
-
-			float DrawScale = 300;
-			FVector2D CentreDrawPosition(350, 400);
-			FVector2D WheelDrawPosition(WheelOffset.Y, -WheelOffset.X);
-			WheelDrawPosition *= DrawScale;
-			WheelDrawPosition /= MaxSize.X;
-			WheelDrawPosition += CentreDrawPosition;
-
-			FVector2D WheelDimensions(Wheel->WheelWidth, Wheel->WheelRadius * 2.0f);
-			FVector2D HalfDimensions = WheelDimensions * 0.5f;
-			FCanvasBoxItem BoxItem(WheelDrawPosition - HalfDimensions, WheelDimensions);
-			BoxItem.SetColor(FColor::Green);
-			Canvas->DrawItem(BoxItem);
-
-			float VisualScaling = 0.0001f;
-			FVector2D Force2D(Forces.Y * VisualScaling, -Forces.X * VisualScaling);
-			DrawLine2D(Canvas, WheelDrawPosition, WheelDrawPosition + Force2D, FColor::Red);
-
-			float SlipAngle = FMath::Abs(PWheel.GetSlipAngle());
-			float X = FMath::Sin(SlipAngle) * 50.f;
-			float Y = FMath::Cos(SlipAngle) * 50.f;
-
-			int Xpos = WheelDrawPosition.X + 20;
-			int Ypos = WheelDrawPosition.Y - 75.f;
-			DrawLine2D(Canvas, WheelDrawPosition, WheelDrawPosition - FVector2D(X, Y), FColor::White);
-			Ypos += Canvas->DrawText(RenderFont, FString::Printf(TEXT("Slip Angle : %d %"), (int)RadToDeg(SlipAngle)), Xpos, Ypos);
-
-			Ypos += Canvas->DrawText(RenderFont, FString::Printf(TEXT("AccelT : %.1f"), PWheel.GetDriveTorque()), Xpos, Ypos);
-			Ypos += Canvas->DrawText(RenderFont, FString::Printf(TEXT("BrakeT : %.1f"), PWheel.GetBrakeTorque()), Xpos, Ypos);
-
-			Ypos += Canvas->DrawText(RenderFont, FString::Printf(TEXT("Omega : %.2f"), PWheel.GetAngularVelocity()), Xpos, Ypos);
-
-			Ypos += Canvas->DrawText(RenderFont, FString::Printf(TEXT("GroundV : %.1f"), PWheel.GetRoadSpeed()), Xpos, Ypos);
-			Ypos += Canvas->DrawText(RenderFont, FString::Printf(TEXT("WheelV : %.1f"), PWheel.GetWheelGroundSpeed()), Xpos, Ypos);
-//			Ypos += Canvas->DrawText(RenderFont, FString::Printf(TEXT("Sx : %.2f"), PWheel.GetNormalizedLongitudinalSlip()), Xpos, Ypos);
-			Ypos += Canvas->DrawText(RenderFont, FString::Printf(TEXT("Long Ad Limit : %.2f"), PWheel.LongitudinalAdhesiveLimit), Xpos, Ypos);
-			Ypos += Canvas->DrawText(RenderFont, FString::Printf(TEXT("Lat Ad Limit : %.2f"), PWheel.LateralAdhesiveLimit), Xpos, Ypos);
-			Ypos += Canvas->DrawText(RenderFont, FString::Printf(TEXT("Side Slip : %.2f"), PWheel.SideSlipModifier), Xpos, Ypos);
-
-			if (PWheel.AppliedLinearDriveForce > PWheel.LongitudinalAdhesiveLimit)
-			{
-				Canvas->SetDrawColor(FColor::Red);
-			}
-			else
-			{
-				Canvas->SetDrawColor(FColor::Green);
-			}
-			Ypos += Canvas->DrawText(RenderFont, FString::Printf(TEXT("Ap Drive : %.2f"), PWheel.AppliedLinearDriveForce), Xpos, Ypos);
-
-			if (PWheel.AppliedLinearBrakeForce > PWheel.LongitudinalAdhesiveLimit)
-			{
-				Canvas->SetDrawColor(FColor::Red);
-			}
-			else
-			{
-				Canvas->SetDrawColor(FColor::Green);
-			}
-			Ypos += Canvas->DrawText(RenderFont, FString::Printf(TEXT("Ap Brake : %.2f"), PWheel.AppliedLinearBrakeForce), Xpos, Ypos);
-			Canvas->SetDrawColor(FColor::White);
-
-			//if (PWheel.Setup().EngineEnabled)
-			//{
-			//	Ypos += Canvas->DrawText(RenderFont, FString::Printf(TEXT("RPM        : %.1f"), PWheel.GetWheelRPM()), Xpos, Ypos);
-			//	Ypos += Canvas->DrawText(RenderFont, FString::Printf(TEXT("Geared RPM : %.1f"), PTransmission.GetEngineRPMFromWheelRPM(PWheel.GetWheelRPM())), Xpos, Ypos);
-
-			//}
-
-			if (ContactMat)
-			{
-				Canvas->DrawText(RenderFont
-					, FString::Printf(TEXT("Friction %d"), ContactMat->Friction)
-					, WheelDrawPosition.X, WheelDrawPosition.Y-95.f);
-			}
+			auto& PWheel = PVehicleOutput->Wheels[WheelIdx];
+//			FVector Forces = PWheel.GetForceFromFriction();
+//
+//			FChaosWheelSetup& WheelSetup = WheelSetups[WheelIdx];
+//			UChaosVehicleWheel* Wheel = WheelSetup.WheelClass.GetDefaultObject();
+//			check(Wheel);
+//			UPhysicalMaterial* ContactMat = Wheel->GetContactSurfaceMaterial();
+//
+//			const FVector WheelOffset = GetWheelRestingPosition(WheelSetup);
+//
+//			float DrawScale = 300;
+//			FVector2D CentreDrawPosition(350, 400);
+//			FVector2D WheelDrawPosition(WheelOffset.Y, -WheelOffset.X);
+//			WheelDrawPosition *= DrawScale;
+//			WheelDrawPosition /= MaxSize.X;
+//			WheelDrawPosition += CentreDrawPosition;
+//
+//			FVector2D WheelDimensions(Wheel->WheelWidth, Wheel->WheelRadius * 2.0f);
+//			FVector2D HalfDimensions = WheelDimensions * 0.5f;
+//			FCanvasBoxItem BoxItem(WheelDrawPosition - HalfDimensions, WheelDimensions);
+//			BoxItem.SetColor(FColor::Green);
+//			Canvas->DrawItem(BoxItem);
+//
+//			float VisualScaling = 0.0001f;
+//			FVector2D Force2D(Forces.Y * VisualScaling, -Forces.X * VisualScaling);
+//			DrawLine2D(Canvas, WheelDrawPosition, WheelDrawPosition + Force2D, FColor::Red);
+//
+//			float SlipAngle = FMath::Abs(PWheel.GetSlipAngle());
+//			float X = FMath::Sin(SlipAngle) * 50.f;
+//			float Y = FMath::Cos(SlipAngle) * 50.f;
+//
+//			int Xpos = WheelDrawPosition.X + 20;
+//			int Ypos = WheelDrawPosition.Y - 75.f;
+//			DrawLine2D(Canvas, WheelDrawPosition, WheelDrawPosition - FVector2D(X, Y), FColor::White);
+//			Ypos += Canvas->DrawText(RenderFont, FString::Printf(TEXT("Slip Angle : %d %"), (int)RadToDeg(SlipAngle)), Xpos, Ypos);
+//
+//			Ypos += Canvas->DrawText(RenderFont, FString::Printf(TEXT("AccelT : %.1f"), PWheel.GetDriveTorque()), Xpos, Ypos);
+//			Ypos += Canvas->DrawText(RenderFont, FString::Printf(TEXT("BrakeT : %.1f"), PWheel.GetBrakeTorque()), Xpos, Ypos);
+//
+//			Ypos += Canvas->DrawText(RenderFont, FString::Printf(TEXT("Omega : %.2f"), PWheel.GetAngularVelocity()), Xpos, Ypos);
+//
+//			Ypos += Canvas->DrawText(RenderFont, FString::Printf(TEXT("GroundV : %.1f"), PWheel.GetRoadSpeed()), Xpos, Ypos);
+//			Ypos += Canvas->DrawText(RenderFont, FString::Printf(TEXT("WheelV : %.1f"), PWheel.GetWheelGroundSpeed()), Xpos, Ypos);
+////			Ypos += Canvas->DrawText(RenderFont, FString::Printf(TEXT("Sx : %.2f"), PWheel.GetNormalizedLongitudinalSlip()), Xpos, Ypos);
+//			Ypos += Canvas->DrawText(RenderFont, FString::Printf(TEXT("Long Ad Limit : %.2f"), PWheel.LongitudinalAdhesiveLimit), Xpos, Ypos);
+//			Ypos += Canvas->DrawText(RenderFont, FString::Printf(TEXT("Lat Ad Limit : %.2f"), PWheel.LateralAdhesiveLimit), Xpos, Ypos);
+//			Ypos += Canvas->DrawText(RenderFont, FString::Printf(TEXT("Side Slip : %.2f"), PWheel.SideSlipModifier), Xpos, Ypos);
+//
+//			if (PWheel.AppliedLinearDriveForce > PWheel.LongitudinalAdhesiveLimit)
+//			{
+//				Canvas->SetDrawColor(FColor::Red);
+//			}
+//			else
+//			{
+//				Canvas->SetDrawColor(FColor::Green);
+//			}
+//			Ypos += Canvas->DrawText(RenderFont, FString::Printf(TEXT("Ap Drive : %.2f"), PWheel.AppliedLinearDriveForce), Xpos, Ypos);
+//
+//			if (PWheel.AppliedLinearBrakeForce > PWheel.LongitudinalAdhesiveLimit)
+//			{
+//				Canvas->SetDrawColor(FColor::Red);
+//			}
+//			else
+//			{
+//				Canvas->SetDrawColor(FColor::Green);
+//			}
+//			Ypos += Canvas->DrawText(RenderFont, FString::Printf(TEXT("Ap Brake : %.2f"), PWheel.AppliedLinearBrakeForce), Xpos, Ypos);
+//			Canvas->SetDrawColor(FColor::White);
+//
+//			//if (PWheel.Setup().EngineEnabled)
+//			//{
+//			//	Ypos += Canvas->DrawText(RenderFont, FString::Printf(TEXT("RPM        : %.1f"), PWheel.GetWheelRPM()), Xpos, Ypos);
+//			//	Ypos += Canvas->DrawText(RenderFont, FString::Printf(TEXT("Geared RPM : %.1f"), PTransmission.GetEngineRPMFromWheelRPM(PWheel.GetWheelRPM())), Xpos, Ypos);
+//
+//			//}
+//
+//			if (ContactMat)
+//			{
+//				Canvas->DrawText(RenderFont
+//					, FString::Printf(TEXT("Friction %d"), ContactMat->Friction)
+//					, WheelDrawPosition.X, WheelDrawPosition.Y-95.f);
+//			}
 		
 		}
 
@@ -1415,97 +1597,95 @@
 
 	if (DebugPage == EDebugPages::SteeringPage)
 	{
-		FVector2D MaxSize = GetWheelLayoutDimensions();
+		//FVector2D MaxSize = GetWheelLayoutDimensions();
 
-		auto& PSteering = PVehicle->GetSteering();
+		//auto& PSteering = PVehicle->GetSteering();
 
-		FVector2D J1, J2;
-		for (int WheelIdx = 0; WheelIdx < PVehicle->Wheels.Num(); WheelIdx++)
-		{
-			FChaosWheelSetup& WheelSetup = WheelSetups[WheelIdx];
-			auto& PWheel = PVehicle->Wheels[WheelIdx];
-			const FVector WheelOffset = GetWheelRestingPosition(WheelSetup);
+		//FVector2D J1, J2;
+		//for (int WheelIdx = 0; WheelIdx < PVehicle->Wheels.Num(); WheelIdx++)
+		//{
+		//	FChaosWheelSetup& WheelSetup = WheelSetups[WheelIdx];
+		//	auto& PWheel = PVehicleOutput->Wheels[WheelIdx];
+		//	const FVector WheelOffset = GetWheelRestingPosition(WheelSetup);
 
-			float Scale = 300.0f / MaxSize.X;
-			FVector2D CentreDrawPosition(450, 400);
-			FVector2D WheelDrawPosition(WheelOffset.Y, -WheelOffset.X);
-			WheelDrawPosition *= Scale;
-			WheelDrawPosition += CentreDrawPosition;
+		//	float Scale = 300.0f / MaxSize.X;
+		//	FVector2D CentreDrawPosition(450, 400);
+		//	FVector2D WheelDrawPosition(WheelOffset.Y, -WheelOffset.X);
+		//	WheelDrawPosition *= Scale;
+		//	WheelDrawPosition += CentreDrawPosition;
 
-			if (PWheel.Setup().SteeringEnabled)
-			{
-				if (WheelOffset.Y > 0)
-				{
-					float SteerAngle = DegToRad(PWheel.GetSteeringAngle());
-					FVector2D Tire = FVector2D(FMath::Sin(SteerAngle), -FMath::Cos(SteerAngle)) * 30.0f;
-					FVector2D WPt = WheelDrawPosition;
-					DrawLine2D(Canvas, WPt - Tire, WPt + Tire, FColor::Black, 8);
+		//	if (PWheel.Setup().SteeringEnabled)
+		//	{
+		//		if (WheelOffset.Y > 0)
+		//		{
+		//			float SteerAngle = DegToRad(PWheel.GetSteeringAngle());
+		//			FVector2D Tire = FVector2D(FMath::Sin(SteerAngle), -FMath::Cos(SteerAngle)) * 30.0f;
+		//			FVector2D WPt = WheelDrawPosition;
+		//			DrawLine2D(Canvas, WPt - Tire, WPt + Tire, FColor::Black, 8);
 
-					if (SteeringSetup.SteeringType == ESteeringType::Ackermann)
-					{
-						FVector2D C1, P, C2;
-						PSteering.Ackermann.GetRightHingeLocations(C1, P, C2);
-						C1.Y = -C1.Y;
-						P.Y = -P.Y;
-						C2.Y = -C2.Y;
+		//			if (SteeringSetup.SteeringType == ESteeringType::Ackermann)
+		//			{
+		//				FVector2D C1, P, C2;
+		//				PSteering.Ackermann.GetRightHingeLocations(C1, P, C2);
+		//				C1.Y = -C1.Y;
+		//				P.Y = -P.Y;
+		//				C2.Y = -C2.Y;
 
-						FVector2D JPt = WheelDrawPosition + (P - C2) * Scale;
-						FVector2D CPt = WheelDrawPosition + (C1 - C2) * Scale;
-						DrawLine2D(Canvas, CPt, JPt, FColor::Orange, 3);
-						DrawLine2D(Canvas, WPt, JPt, FColor::Orange, 3);
-						J1 = CPt;
-					}
-				}
-				else
-				{
-					float SteerAngle = DegToRad(PWheel.GetSteeringAngle());
-					FVector2D Tire = FVector2D(FMath::Sin(SteerAngle), -FMath::Cos(SteerAngle)) * 30.0f;
-					FVector2D WPt = WheelDrawPosition;
-					DrawLine2D(Canvas, WPt - Tire, WPt + Tire, FColor::Black, 8);
+		//				FVector2D JPt = WheelDrawPosition + (P - C2) * Scale;
+		//				FVector2D CPt = WheelDrawPosition + (C1 - C2) * Scale;
+		//				DrawLine2D(Canvas, CPt, JPt, FColor::Orange, 3);
+		//				DrawLine2D(Canvas, WPt, JPt, FColor::Orange, 3);
+		//				J1 = CPt;
+		//			}
+		//		}
+		//		else
+		//		{
+		//			float SteerAngle = DegToRad(PWheel.GetSteeringAngle());
+		//			FVector2D Tire = FVector2D(FMath::Sin(SteerAngle), -FMath::Cos(SteerAngle)) * 30.0f;
+		//			FVector2D WPt = WheelDrawPosition;
+		//			DrawLine2D(Canvas, WPt - Tire, WPt + Tire, FColor::Black, 8);
 
-					if (SteeringSetup.SteeringType == ESteeringType::Ackermann)
-					{
+		//			if (SteeringSetup.SteeringType == ESteeringType::Ackermann)
+		//			{
 
-						FVector2D C1, P, C2;
-						PSteering.Ackermann.GetLeftHingeLocations(C1, P, C2);
-						C1.Y = -C1.Y;
-						P.Y = -P.Y;
-						C2.Y = -C2.Y;
+		//				FVector2D C1, P, C2;
+		//				PSteering.Ackermann.GetLeftHingeLocations(C1, P, C2);
+		//				C1.Y = -C1.Y;
+		//				P.Y = -P.Y;
+		//				C2.Y = -C2.Y;
 
-						FVector2D JPt = WheelDrawPosition + (P - C2) * Scale;
-						FVector2D CPt = WheelDrawPosition + (C1 - C2) * Scale;
-						DrawLine2D(Canvas, CPt, JPt, FColor::Orange, 3);
-						DrawLine2D(Canvas, WPt, JPt, FColor::Orange, 3);
-						J2 = CPt;
-					}
-				}
-			}
-			else
-			{
-				FVector2D CPt = WheelDrawPosition;
-				FVector2D Tire = FVector2D(0.f, 30.0f);
-				DrawLine2D(Canvas, CPt - Tire, CPt + Tire, FColor::Black, 8);
-			}
+		//				FVector2D JPt = WheelDrawPosition + (P - C2) * Scale;
+		//				FVector2D CPt = WheelDrawPosition + (C1 - C2) * Scale;
+		//				DrawLine2D(Canvas, CPt, JPt, FColor::Orange, 3);
+		//				DrawLine2D(Canvas, WPt, JPt, FColor::Orange, 3);
+		//				J2 = CPt;
+		//			}
+		//		}
+		//	}
+		//	else
+		//	{
+		//		FVector2D CPt = WheelDrawPosition;
+		//		FVector2D Tire = FVector2D(0.f, 30.0f);
+		//		DrawLine2D(Canvas, CPt - Tire, CPt + Tire, FColor::Black, 8);
+		//	}
 
-			Canvas->DrawText(RenderFont
-				, FString::Printf(TEXT("Angle %.1f"), PWheel.GetSteeringAngle())
-				, WheelDrawPosition.X, WheelDrawPosition.Y - 15.f);
+		//	Canvas->DrawText(RenderFont
+		//		, FString::Printf(TEXT("Angle %.1f"), PWheel.GetSteeringAngle())
+		//		, WheelDrawPosition.X, WheelDrawPosition.Y - 15.f);
 
-		}
-		DrawLine2D(Canvas, J1, J2, FColor::Red, 3);
+		//}
+		//DrawLine2D(Canvas, J1, J2, FColor::Red, 3);
 
 	}
 
 	// draw engine torque curve - just putting engine under transmission
-	if (DebugPage == EDebugPages::TransmissionPage && PVehicle->HasTransmission())
+	if (DebugPage == EDebugPages::TransmissionPage && bMechanicalSimEnabled)
 	{
-		auto& PTransmission = PVehicle->GetTransmission();
-		auto& PEngine = PVehicle->GetEngine();
-		auto& PTransmissionSetup = PTransmission.Setup();
+		auto& PTransmissionSetup = TransmissionSetup;
 
-		float MaxTorque = PEngine.Setup().MaxTorque;
-		int CurrentRPM = (int)PEngine.GetEngineRPM();
-		FVector2D CurrentValue(CurrentRPM, PEngine.GetEngineTorque());
+		float MaxTorque = EngineSetup.MaxTorque;
+		int CurrentRPM = (int)PVehicleOutput->EngineRPM;
+		FVector2D CurrentValue(CurrentRPM, PVehicleOutput->EngineTorque);
 		int GraphWidth = 200; int GraphHeight = 120;
 		int GraphXPos = 200; int GraphYPos = 400;
 
@@ -1513,13 +1693,14 @@
 			, CurrentValue.X, CurrentValue.Y
 			, GraphXPos, GraphYPos, 
 			GraphWidth, GraphHeight, 
-			FVector2D(0, PEngine.Setup().MaxRPM), FVector2D(MaxTorque, 0));
+			FVector2D(0, EngineSetup.MaxRPM), FVector2D(MaxTorque, 0));
 
 		FVector2D LastPoint;
-		for (float RPM = 0; RPM <= PEngine.Setup().MaxRPM; RPM += 10.f)
+		for (float RPM = 0; RPM <= EngineSetup.MaxRPM; RPM += 10.f)
 		{
-			float X = RPM / PEngine.Setup().MaxRPM;
-			float Y = PEngine.GetTorqueFromRPM(RPM, false) / MaxTorque;
+			float X = RPM / EngineSetup.MaxRPM;
+			float Y = EngineSetup.GetTorqueFromRPM(RPM) / MaxTorque;
+			//float Y = PEngine.GetTorqueFromRPM(RPM, false) / MaxTorque;
 			FVector2D NextPoint(GraphXPos + GraphWidth * X, GraphYPos + GraphHeight - GraphHeight * Y);
 			if (RPM > SMALL_NUMBER)
 			{
@@ -1529,29 +1710,25 @@
 		}
 
 		Canvas->DrawText(RenderFont
-			, FString::Printf(TEXT("RevRate %.1f"), PEngine.GetEngineRevRate())
+			, FString::Printf(TEXT("RevRate %.1f"), EngineSetup.EngineRevDownRate)
 			, GraphXPos, GraphYPos);
 
 	}
 
 	// draw transmission torque curve
-	if (DebugPage == EDebugPages::TransmissionPage && PVehicle->HasTransmission())
+	if (DebugPage == EDebugPages::TransmissionPage && bMechanicalSimEnabled)
 	{
-		auto& PTransmission = PVehicle->GetTransmission();
-		auto& PEngine = PVehicle->GetEngine();
-		auto& PTransmissionSetup = PTransmission.Setup();
-
-		auto& ESetup = PEngine.Setup();
-		auto& TSetup = PTransmission.Setup();
+		auto& ESetup = EngineSetup;
+		auto& TSetup = TransmissionSetup;
 		float MaxTorque = ESetup.MaxTorque;
-		float MaxGearRatio = TSetup.ForwardRatios[0] * TSetup.FinalDriveRatio; // 1st gear always has the highest multiplier
-		float LongGearRatio = TSetup.ForwardRatios[TSetup.ForwardRatios.Num()-1] * TSetup.FinalDriveRatio;
+		float MaxGearRatio = TSetup.ForwardGearRatios[0] * TSetup.FinalRatio; // 1st gear always has the highest multiplier
+		float LongGearRatio = TSetup.ForwardGearRatios[TSetup.ForwardGearRatios.Num()-1] * TSetup.FinalRatio;
 		int GraphWidth = 400; int GraphHeight = 240;
 		int GraphXPos = 500; int GraphYPos = 150;
 
 		{
-			float X = PTransmission.GetTransmissionRPM();
-			float Y = PTransmission.GetTransmissionTorque(PEngine.GetTorqueFromRPM(false));
+			float X = PVehicleOutput->TransmissionRPM;
+			float Y = PVehicleOutput->TransmissionTorque;
 
 			FVector2D CurrentValue(X, Y);
 			Canvas->DrawDebugGraph(FString("Transmission Torque Graph")
@@ -1563,14 +1740,14 @@
 
 		FVector2D LastPoint;
 
-		for (int Gear = 1; Gear <= TSetup.ForwardRatios.Num(); Gear++)
+		for (int Gear = 1; Gear <= TSetup.ForwardGearRatios.Num(); Gear++)
 		{
 			for (int EngineRPM = 0; EngineRPM <= ESetup.MaxRPM; EngineRPM += 10)
 			{
-				float RPMOut = PTransmission.GetTransmissionRPM(EngineRPM, Gear);
+				float RPMOut = EngineRPM / TSetup.GetGearRatio(Gear);
 
 				float X = RPMOut / (ESetup.MaxRPM / LongGearRatio);
-				float Y = PEngine.GetTorqueFromRPM(EngineRPM, false) * PTransmission.GetGearRatio(Gear) / (MaxTorque*MaxGearRatio);
+				float Y = ESetup.GetTorqueFromRPM(EngineRPM) * TSetup.GetGearRatio(Gear) / (MaxTorque*MaxGearRatio);
 				FVector2D NextPoint(GraphXPos + GraphWidth * X, GraphYPos + GraphHeight - GraphHeight * Y);
 				if (EngineRPM > 0)
 				{
@@ -1588,7 +1765,6 @@
 
 		for (int32 WheelIdx = 0; WheelIdx < WheelSetups.Num(); ++WheelIdx)
 		{
-			auto& PSuspension = PVehicle->Suspension[WheelIdx];
 			FChaosWheelSetup& WheelSetup = WheelSetups[WheelIdx];
 			UChaosVehicleWheel* Wheel = WheelSetup.WheelClass.GetDefaultObject();
 			check(Wheel);
@@ -1609,9 +1785,9 @@
 				FVector2D End = Start + FVector2D(20.f, 0.f);
 				DrawLine2D(Canvas, Start, End, FColor::Yellow, 2.f);
 			}
-
-			float Raise = PSuspension.Setup().SuspensionMaxRaise;
-			float Drop = PSuspension.Setup().SuspensionMaxDrop;
+	
+			float Raise = VehicleWheel->SuspensionMaxRaise;
+			float Drop = VehicleWheel->SuspensionMaxDrop;
 			float Scale = 5.0f;
 
 			{
@@ -1640,85 +1816,6 @@
 		}
 	}
 #endif
-
-}
-
-void UChaosWheeledVehicleMovementComponent::DrawDebug3D()
-{
-#if !(UE_BUILD_SHIPPING || UE_BUILD_TEST)
-
-	Super::DrawDebug3D();
-
-	FBodyInstance* TargetInstance = GetBodyInstance();
-	if (TargetInstance == nullptr)
-	{
-		return;
-	}
-
-	const FTransform BodyTransform = VehicleState.VehicleWorldTransform;
-
-	if (GWheeledVehicleDebugParams.ShowSuspensionLimits)
-	{
-		for (int WheelIdx = 0; WheelIdx < PVehicle->Suspension.Num(); WheelIdx++)
-		{
-			auto& PSuspension = PVehicle->Suspension[WheelIdx];
-			auto& PWheel = PVehicle->Wheels[WheelIdx];
-			// push the visualization out a bit sideways from the wheel model so we can actually see it
-			FVector VehicleRightAxis = VehicleState.VehicleWorldTransform.GetUnitAxis(EAxis::Y) * 48.0f;
-			const FVector& WheelOffset = PSuspension.GetLocalRestingPosition();
-			if (WheelOffset.Y < 0.0f)
-			{
-				VehicleRightAxis = VehicleRightAxis * -1.0f;
-			}
-
-			FVector LocalDirection = PSuspension.Setup().SuspensionAxis;
-			FVector WorldLocation = BodyTransform.TransformPosition(WheelOffset);
-			FVector WorldDirection = BodyTransform.TransformVector(LocalDirection);
-
-			FVector Start = WorldLocation + WorldDirection * (PWheel.GetEffectiveRadius() - PSuspension.Setup().SuspensionMaxRaise);
-			FVector End = WorldLocation + WorldDirection * (PWheel.GetEffectiveRadius() + PSuspension.Setup().SuspensionMaxDrop);
-
-			DrawDebugLine(GetWorld(), Start + VehicleRightAxis, End + VehicleRightAxis, FColor::Orange, false, -1.f, 0, 3.f);
-
-			FVector Start2 = WorldLocation - WorldDirection * PSuspension.Setup().SuspensionMaxRaise;
-			FVector End2 = WorldLocation + WorldDirection * PSuspension.Setup().SuspensionMaxDrop;
-
-			DrawDebugLine(GetWorld(), Start2 + VehicleRightAxis, End2 + VehicleRightAxis, FColor::Yellow, false, -1.f, 0, 3.f);
-		}
-	}
-
-	if (GWheeledVehicleDebugParams.ShowWheelCollisionNormal)
-	{
-		for (int WheelIdx = 0; WheelIdx < Wheels.Num(); WheelIdx++)
-		{
-			FHitResult& Hit = Wheels[WheelIdx]->HitResult;
-			DrawDebugLine(GetWorld(), Hit.ImpactPoint, Hit.ImpactPoint + Hit.Normal * 20.0f, FColor::Yellow, true, 1.0f, 0, 1.0f);
-		}
-	}
-
-	if (GWheeledVehicleDebugParams.ShowSuspensionRaycasts)
-	{
-		for (int WheelIdx = 0; WheelIdx < Wheels.Num(); WheelIdx++)
-		{
-			const FVector& TraceStart = WheelState.Trace[WheelIdx].Start;
-			const FVector& TraceEnd = WheelState.Trace[WheelIdx].End;
-
-			// push the visualization out a bit sideways from the wheel model so we can actually see it
-			FVector VehicleRightAxis = VehicleState.VehicleWorldTransform.GetUnitAxis(EAxis::Y) * 50.0f;
-			const FVector& WheelOffset = PVehicle->Suspension[WheelIdx].GetLocalRestingPosition();
-			if (WheelOffset.Y < 0.0f)
-			{
-				VehicleRightAxis = VehicleRightAxis * -1.0f;
-			}
-
-			FColor UseColor = PVehicle->Wheels[WheelIdx].InContact() ? FColor::Green : FColor::Red;
-			DrawDebugDirectionalArrow(GetWorld(), TraceStart + VehicleRightAxis, TraceEnd + VehicleRightAxis, 10.f, UseColor, false, -1.f, 0, 2.f);
-
-			DrawDebugLine(GetWorld(), TraceStart, TraceStart + VehicleRightAxis, FColor::White, false, -1.f, 0, 1.f);
-			DrawDebugLine(GetWorld(), TraceEnd, TraceEnd + VehicleRightAxis, FColor::White, false, -1.f, 0, 1.f);
-		}
-	}
-#endif
 }
 
 #if !(UE_BUILD_SHIPPING || UE_BUILD_TEST)
@@ -1754,26 +1851,23 @@
 {
 	for (int WheelIdx = 0; WheelIdx < Wheels.Num(); WheelIdx++)
 	{
-		FChaosWheelSetup& WheelSetup = WheelSetups[WheelIdx];
-		UChaosVehicleWheel* Wheel = WheelSetup.WheelClass.GetDefaultObject();
-		auto& PWheel = PVehicle->Wheels[WheelIdx];
-		auto& PSuspension = PVehicle->Suspension[WheelIdx];
+		auto& PWheel = PVehicleOutput->Wheels[WheelIdx];
 		FHitResult& HitResult = Wheels[WheelIdx]->HitResult;
 
 		FWheelStatus& State = WheelStatus[WheelIdx];
-	
+
 		State.bInContact = HitResult.bBlockingHit;
 		State.ContactPoint = HitResult.ImpactPoint;
 		State.PhysMaterial = HitResult.PhysMaterial;
-		State.NormalizedSuspensionLength = PSuspension.GetNormalizedLength();
-		State.SpringForce = PSuspension.GetSuspensionForce();
-		State.bIsSlipping = PWheel.IsSlipping();
-		State.SlipMagnitude = PWheel.GetSlipMagnitude();
-		State.bIsSkidding = PWheel.IsSkidding();
-		State.SkidMagnitude = PWheel.GetSkidMagnitude();
+		State.NormalizedSuspensionLength = PWheel.NormalizedSuspensionLength;
+		State.SpringForce = PWheel.SpringForce;
+		State.bIsSlipping = PWheel.bIsSlipping;
+		State.SlipMagnitude = PWheel.SlipMagnitude;
+		State.bIsSkidding = PWheel.bIsSkidding;
+		State.SkidMagnitude = PWheel.SkidMagnitude;
 		if (State.bIsSkidding)
 		{
-			State.SkidNormal = WheelState.WorldWheelVelocity[WheelIdx].GetSafeNormal();
+			State.SkidNormal = PWheel.SkidNormal;
 			//DrawDebugLine(GetWorld()
 			//	, State.ContactPoint
 			//	, State.ContactPoint + State.SkidNormal
@@ -1786,6 +1880,7 @@
 	}
 }
 
+
 void UChaosWheeledVehicleMovementComponent::BreakWheelStatus(const struct FWheelStatus& Status, bool& bInContact, FVector& ContactPoint, UPhysicalMaterial*& PhysMaterial
 	, float& NormalizedSuspensionLength, float& SpringForce, bool& bIsSlipping, float& SlipMagnitude, bool& bIsSkidding, float& SkidMagnitude, FVector& SkidNormal)
 {
@@ -1819,6 +1914,12 @@
 	return Status;
 }
 
+void UChaosWheeledVehicleMovementComponent::ParallelUpdate(float DeltaSeconds)
+{
+	UChaosVehicleMovementComponent::ParallelUpdate(DeltaSeconds);\
 
+	FillWheelOutputState(); // expose wheel/suspensionm data to blueprint
+}
+
 #endif
 
diff --git a/Engine/Plugins/Experimental/ChaosVehiclesPlugin/Source/ChaosVehicles/Public/AsyncCallback.h b/Engine/Plugins/Experimental/ChaosVehiclesPlugin/Source/ChaosVehicles/Public/AsyncCallback.h
new file mode 100644
--- /dev/null
+++ b/Engine/Plugins/Experimental/ChaosVehiclesPlugin/Source/ChaosVehicles/Public/AsyncCallback.h
@@ -0,0 +1,164 @@
+// Copyright Epic Games, Inc. All Rights Reserved.
+
+#pragma once
+
+#include "CoreMinimal.h"
+#include "PhysicsPublic.h"
+#include "Chaos/SimCallbackInput.h"
+#include "Chaos/SimCallbackObject.h"
+#include "Chaos/GeometryParticlesfwd.h"
+
+class UChaosVehicleMovementComponent;
+
+enum EChaosAsyncVehicleDataType : int8
+{
+	AsyncInvalid,
+	AsyncDefault,
+};
+
+struct FWheelsOutput
+{
+	FWheelsOutput()
+		: InContact(false)
+		, SteeringAngle(0.f)
+		, AngularPosition(0.f)
+		, LateralAdhesiveLimit(0.f)
+		, LongitudinalAdhesiveLimit(0.f)
+		, bIsSlipping(false)
+		, SlipMagnitude(0.f)
+		, bIsSkidding(false)
+		, SkidMagnitude(0.f)
+		, SkidNormal(FVector(1,0,0))
+		, SuspensionOffset(0.f)
+		, SpringForce(0.f)
+		, NormalizedSuspensionLength(0.f)
+	{
+	}
+
+	// wheels
+	bool InContact;
+	float SteeringAngle;
+	float AngularPosition;
+
+	float LateralAdhesiveLimit;
+	float LongitudinalAdhesiveLimit;
+
+	bool bIsSlipping;
+	float SlipMagnitude;
+	bool bIsSkidding;
+	float SkidMagnitude;
+	FVector SkidNormal;
+
+	// suspension related
+	float SuspensionOffset;
+	float SpringForce;
+	float NormalizedSuspensionLength;
+
+};
+
+/**
+ * Per Vehicle Output State from Physics Thread to Game Thread                            
+ */
+struct FPhysicsVehicleOutput
+{
+	FPhysicsVehicleOutput()
+		: CurrentGear(0)
+		, TargetGear(0)
+		, EngineRPM(0.f)
+		, EngineTorque(0.f)
+		, TransmissionRPM(0.f)
+		, TransmissionTorque(0.f)
+	{
+	}
+
+	TArray<FWheelsOutput> Wheels;
+	int32 CurrentGear;
+	int32 TargetGear;
+	float EngineRPM;
+	float EngineTorque;
+	float TransmissionRPM;
+	float TransmissionTorque;
+};
+
+/**
+ * Per Vehicle Input State from Game Thread to Physics Thread
+ */                             
+struct FChaosVehicleAsyncInput
+{
+	const EChaosAsyncVehicleDataType Type;
+	UChaosVehicleMovementComponent* Vehicle;
+	
+	union {
+		FRigidParticlePhysicsProxy* Proxy;
+	} Actor;
+
+	/** 
+	* Vehicle simulation running on the Physics Thread
+	*/
+	virtual TUniquePtr<struct FChaosVehicleAsyncOutput> Simulate(UWorld* World, const float DeltaSeconds, const float TotalSeconds, bool& bWakeOut) const = 0;
+
+	FChaosVehicleAsyncInput(EChaosAsyncVehicleDataType InType = EChaosAsyncVehicleDataType::AsyncInvalid)
+		: Type(InType)
+		, Vehicle(nullptr)
+	{
+		Actor.Proxy = nullptr;	//indicates async/sync task not needed
+	}
+
+	virtual ~FChaosVehicleAsyncInput() = default;
+};
+
+struct FChaosVehicleManagerAsyncInput : public Chaos::FSimCallbackInput
+{
+	TArray<TUniquePtr<FChaosVehicleAsyncInput>> VehicleInputs;
+
+	TWeakObjectPtr<UWorld> World;
+	int32 Timestamp = INDEX_NONE;
+
+	void Reset()
+	{
+		VehicleInputs.Reset();
+		World.Reset();
+	}
+};
+
+/**
+ * Async Output Data
+ */
+struct FChaosVehicleAsyncOutput
+{
+	const EChaosAsyncVehicleDataType Type;
+	bool bValid;	// indicates no work was done
+	FPhysicsVehicleOutput VehicleSimOutput;
+
+	FChaosVehicleAsyncOutput(EChaosAsyncVehicleDataType InType = EChaosAsyncVehicleDataType::AsyncInvalid)
+		: Type(InType)
+		, bValid(false)
+	{ }
+
+	virtual ~FChaosVehicleAsyncOutput() = default;
+};
+
+
+/**
+ * Async Output for all of the vehicles handled by this Vehicle Manager
+ */
+struct FChaosVehicleManagerAsyncOutput : public Chaos::FSimCallbackOutput
+{
+	TArray<TUniquePtr<FChaosVehicleAsyncOutput>> VehicleOutputs;
+	int32 Timestamp = INDEX_NONE;
+
+	void Reset()
+	{
+		VehicleOutputs.Reset();
+	}
+};
+
+/**
+ * Async callback from the Physics Engine where we can perform our vehicle simulation
+ */
+class FChaosVehicleManagerAsyncCallback : public Chaos::TSimCallbackObject<FChaosVehicleManagerAsyncInput, FChaosVehicleManagerAsyncOutput>
+{
+private:
+	virtual void OnPreSimulate_Internal() override;
+	virtual void OnContactModification_Internal(const TArrayView<Chaos::FPBDCollisionConstraintHandleModification>& Modifications) override;
+};
diff --git a/Engine/Plugins/Experimental/ChaosVehiclesPlugin/Source/ChaosVehicles/Public/ChaosVehicleManager.h b/Engine/Plugins/Experimental/ChaosVehiclesPlugin/Source/ChaosVehicles/Public/ChaosVehicleManager.h
--- a/Engine/Plugins/Experimental/ChaosVehiclesPlugin/Source/ChaosVehicles/Public/ChaosVehicleManager.h
+++ b/Engine/Plugins/Experimental/ChaosVehiclesPlugin/Source/ChaosVehicles/Public/ChaosVehicleManager.h
@@ -3,30 +3,35 @@
 #pragma once
 
 #include "CoreMinimal.h"
-#include "ChaosVehicleMovementComponent.h"
 #include "PhysicsPublic.h"
 
+#include "Chaos/SimCallbackInput.h"
+#include "Chaos/SimCallbackObject.h"
+#include "Chaos/GeometryParticlesfwd.h"
+#include "AsyncCallback.h"
+
 class UChaosTireConfig;
 class UChaosVehicleMovementComponent;
 
 
-// #todo: We probably want this to be associated with a physics solver. So simulation can be isolated if desired.
-// Physics scene based currently
-
-/**
- * Manages vehicles and tire surface data for all scenes
- */
 class CHAOSVEHICLES_API FChaosVehicleManager
 {
 public:
-
 	// Updated when vehicles need to recreate their physics state.
-	// Used when designer tweaks values while the game is running.
+	// Used when values tweaked while the game is running.
 	static uint32 VehicleSetupTag;
 
 	FChaosVehicleManager(FPhysScene* PhysScene);
 	~FChaosVehicleManager();
 
+	static void OnPostWorldInitialization(UWorld* InWorld, const UWorld::InitializationValues);
+	static void OnWorldCleanup(UWorld* InWorld, bool bSessionEnded, bool bCleanupResources);
+
+#if WITH_CHAOS
+	/** Get Physics Scene */
+	FPhysScene_Chaos& GetScene() const { return Scene; }
+#endif
+
 	/**
 	 * Register a Physics vehicle for processing
 	 */
@@ -38,37 +43,39 @@ public:
 	void RemoveVehicle( TWeakObjectPtr<UChaosVehicleMovementComponent> Vehicle );
 
 	/**
-	 * Update vehicle data before the scene simulates
-	 */
-	void Update(FPhysScene* PhysScene, float DeltaTime);
-	
-	/**
-	 * Update vehicle tuning and other state such as input 
+	 * Update vehicle tuning and other state such as input
 	 */
-	void PreTick(FPhysScene* PhysScene, float DeltaTime);
+	void ScenePreTick(FPhysScene* PhysScene, float DeltaTime);
 
 	/** Detach this vehicle manager from a FPhysScene (remove delegates, remove from map etc) */
 	void DetachFromPhysScene(FPhysScene* PhysScene);
 
-#if WITH_CHAOS
-	/** Get Physics Scene */
-	FPhysScene_Chaos& GetScene() const { return Scene; }
-#endif
+	/*
+	* Flow:
+	*	Update .. SubStep .. PostUpdate
+	*/
+
+	void Update(FPhysScene* PhysScene, float DeltaTime);
+	void SubStep(FPhysScene* PhysScene, float DeltaTime);
+	void PostUpdate(FChaosScene* PhysScene);
+
+	void ParallelUpdateVehicles(float DeltaSeconds);
 
 	/** Find a vehicle manager from an FPhysScene */
 	static FChaosVehicleManager* GetVehicleManagerFromScene(FPhysScene* PhysScene);
 
-	/** Gets a transient default TireConfig object */
-	static UChaosTireConfig* GetDefaultTireConfig();
+	void RegisterCallbacks();
+	void UnregisterCallbacks();
 
 private:
-
 	/** Map of physics scenes to corresponding vehicle manager */
 	static TMap<FPhysScene*, FChaosVehicleManager*> SceneToVehicleManagerMap;
 
 #if WITH_CHAOS
 	// The physics scene we belong to
 	FPhysScene_Chaos& Scene;
+
+	static bool GInitialized;
 #endif
 
 	// All instanced vehicles
@@ -76,5 +83,17 @@ private:
 
 	FDelegateHandle OnPhysScenePreTickHandle;
 	FDelegateHandle OnPhysSceneStepHandle;
+	FDelegateHandle OnPhysScenePostTickHandle;
+
+	static FDelegateHandle OnPostWorldInitializationHandle;
+	static FDelegateHandle OnWorldCleanupHandle;
 
+
+	FChaosVehicleManagerAsyncCallback* AsyncCallback;	// Async callback from the physics engine - we can run our simulation here
+	int32 Timestamp;
+	int32 SubStepCount;
+
+	TArray<Chaos::TSimCallbackOutputHandle<FChaosVehicleManagerAsyncOutput>> PendingOutputs;
+	Chaos::TSimCallbackOutputHandle<FChaosVehicleManagerAsyncOutput> LatestOutput;
 };
+
diff --git a/Engine/Plugins/Experimental/ChaosVehiclesPlugin/Source/ChaosVehicles/Public/ChaosVehicleMovementComponent.h b/Engine/Plugins/Experimental/ChaosVehiclesPlugin/Source/ChaosVehicles/Public/ChaosVehicleMovementComponent.h
--- a/Engine/Plugins/Experimental/ChaosVehiclesPlugin/Source/ChaosVehicles/Public/ChaosVehicleMovementComponent.h
+++ b/Engine/Plugins/Experimental/ChaosVehiclesPlugin/Source/ChaosVehicles/Public/ChaosVehicleMovementComponent.h
@@ -13,7 +13,7 @@
 #include "AerofoilSystem.h"
 #include "ThrustSystem.h"
 
-#include "SimpleVehicle.h" // #todo: move
+#include "AsyncCallback.h"
 
 #include "ChaosVehicleMovementComponent.generated.h"
 
@@ -23,6 +23,60 @@
 
 class UCanvas;
 
+struct FChaosVehicleAsyncInput;
+struct FChaosVehicleManagerAsyncOutput;
+
+
+struct FControlInputs
+{
+	FControlInputs()
+		: SteeringInput(0.f)
+		, ThrottleInput(0.f)
+		, BrakeInput(0.f)
+		, PitchInput(0.f)
+		, RollInput(0.f)
+		, YawInput(0.f)
+		, HandbrakeInput(0.f)
+		, TargetGearInput(0)
+		, SetGearImmediate(-1)
+		, TransmissionType(Chaos::ETransmissionType::Automatic)
+		, GearUpInput(false)
+		, GearDownInput(false)
+	{
+
+	}
+
+	// Steering output to physics system. Range -1...1
+	float SteeringInput;
+
+	// Accelerator output to physics system. Range 0...1
+	float ThrottleInput;
+
+	// Brake output to physics system. Range 0...1
+	float BrakeInput;
+
+	// Body Pitch output to physics system. Range -1...1
+	float PitchInput;
+
+	// Body Roll output to physics system. Range -1...1
+	float RollInput;
+
+	// Body Yaw output to physics system. Range -1...1
+	float YawInput;
+
+	// Handbrake output to physics system. Range 0...1
+	float HandbrakeInput;
+
+	int32 TargetGearInput;
+	int8 SetGearImmediate;
+	Chaos::ETransmissionType TransmissionType;
+
+
+	bool GearUpInput;
+	bool GearDownInput;
+};
+
+
 struct FVehicleDebugParams
 {
 	bool ShowCOM = false;
@@ -38,7 +92,8 @@
 	bool BatchQueries = false;	// Turned off due to Issue with Overlap Queries on scaled terrain
 	float ForceDebugScaling = 0.0006f;
 	float SleepCounterThreshold = 15;
-	bool DisableVehicleSleep = true;
+	bool DisableVehicleSleep = false;
+	bool EnableMultithreading = false; // possibly not all entirely thread safe at present
 };
 
 struct FBodyInstance;
@@ -90,9 +145,22 @@
 
 };
 
+struct FChaosVehicleDefaultAsyncInput : public FChaosVehicleAsyncInput
+{
+	float GravityZ;
+	FControlInputs ControlInputs;
+	FCollisionQueryParams TraceParams;
+
+	FChaosVehicleDefaultAsyncInput();
+
+	TUniquePtr<FChaosVehicleAsyncOutput> Simulate(UWorld* World, const float DeltaSeconds, const float TotalSeconds, bool& bWakeOut) const override;
+};
+
+
 USTRUCT()
 struct FVehicleTorqueControlConfig
 {
+public:
 	GENERATED_USTRUCT_BODY()
 
 	/** Torque Control Enabled */
@@ -124,6 +192,12 @@
 	UPROPERTY(EditAnywhere, Category = Setup)
 	float RotationDamping;
 
+	const Chaos::FTorqueControlConfig& GetTorqueControlConfig()
+	{
+		FillTorqueControlSetup();
+		return PTorqueControlConfig;
+	}
+
 	void InitDefaults()
 	{
 		Enabled = false;
@@ -135,11 +209,28 @@
 		RollFromSteering = 0.0f;
 		RotationDamping = 0.02f;
 	}
+
+private:
+	void FillTorqueControlSetup()
+	{
+		PTorqueControlConfig.Enabled = Enabled;
+		PTorqueControlConfig.YawTorqueScaling = YawTorqueScaling;
+		PTorqueControlConfig.YawFromSteering = YawFromSteering;
+		PTorqueControlConfig.YawFromRollTorqueScaling = YawFromRollTorqueScaling;
+		PTorqueControlConfig.PitchTorqueScaling = PitchTorqueScaling;
+		PTorqueControlConfig.RollTorqueScaling = RollTorqueScaling;
+		PTorqueControlConfig.RollFromSteering = RollFromSteering;
+		PTorqueControlConfig.RotationDamping = RotationDamping;
+
+	}
+
+	Chaos::FTorqueControlConfig PTorqueControlConfig;
 };
 
 USTRUCT()
 struct FVehicleTargetRotationControlConfig
 {
+public:
 	GENERATED_USTRUCT_BODY()
 
 	/** Rotation Control Enabled */
@@ -182,6 +273,11 @@
 	UPROPERTY(EditAnywhere, Category = Setup)
 	float AutoCentreYawStrength;
 
+	const Chaos::FTargetRotationControlConfig& GetTargetRotationControlConfig()
+	{
+		FillTargetRotationControlSetup();
+		return PTargetRotationControlConfig;
+	}
 
 	void InitDefaults()
 	{
@@ -202,11 +298,31 @@
 		AutoCentrePitchStrength = 0.f;
 		AutoCentreYawStrength = 0.f;
 	}
+
+private:
+	void FillTargetRotationControlSetup()
+	{
+		PTargetRotationControlConfig.Enabled = Enabled;
+		PTargetRotationControlConfig.bRollVsSpeedEnabled = bRollVsSpeedEnabled;
+		PTargetRotationControlConfig.RollControlScaling = RollControlScaling;
+		PTargetRotationControlConfig.RollMaxAngle = RollMaxAngle;
+		PTargetRotationControlConfig.PitchControlScaling = PitchControlScaling;
+		PTargetRotationControlConfig.PitchMaxAngle = PitchMaxAngle;
+		PTargetRotationControlConfig.RotationStiffness = RotationStiffness;
+		PTargetRotationControlConfig.RotationDamping = RotationDamping;
+		PTargetRotationControlConfig.MaxAccel = MaxAccel;
+		PTargetRotationControlConfig.AutoCentreRollStrength = AutoCentreRollStrength;
+		PTargetRotationControlConfig.AutoCentrePitchStrength = AutoCentrePitchStrength;
+		PTargetRotationControlConfig.AutoCentreYawStrength = AutoCentreYawStrength;
+	}
+
+	Chaos::FTargetRotationControlConfig PTargetRotationControlConfig;
 };
 
 USTRUCT()
 struct FVehicleStabilizeControlConfig
 {
+public:
 	GENERATED_USTRUCT_BODY()
 
 	/** Torque Control Enabled */
@@ -220,12 +336,29 @@
 	UPROPERTY(EditAnywhere, Category = Setup)
 	float PositionHoldXY;
 
+	const Chaos::FStabilizeControlConfig& GetStabilizeControlConfig()
+	{
+		FillStabilizeControlSetup();
+		return PStabilizeControlConfig;
+	}
+
 	void InitDefaults()
 	{
 		Enabled = false;
 		AltitudeHoldZ = 4.0f;
 		PositionHoldXY = 8.0f;
 	}
+
+private:
+	void FillStabilizeControlSetup()
+	{
+		PStabilizeControlConfig.Enabled = this->Enabled;
+		PStabilizeControlConfig.AltitudeHoldZ = this->AltitudeHoldZ;
+		PStabilizeControlConfig.PositionHoldXY = this->PositionHoldXY;
+	}
+
+	Chaos::FStabilizeControlConfig PStabilizeControlConfig;
+
 };
 
 /** Commonly used state - evaluated once used wherever required */
@@ -255,8 +388,11 @@
 
 	}
 
-	/** Cache some useful data at the start of the frame */
-	void CaptureState(FBodyInstance* TargetInstance, float GravityZ, float DeltaTime);
+	/** Cache some useful data at the start of the frame from GT BodyInstance */
+	void CaptureState(const FBodyInstance* TargetInstance, float GravityZ, float DeltaTime);
 
+	/** Cache some useful data at the start of the frame from Physics thread Particle Handle */
+	void CaptureState(const Chaos::TPBDRigidParticleHandle<float, 3>* Handle, float GravityZ, float DeltaTime);
+
 	FTransform VehicleWorldTransform;
 	FVector VehicleWorldVelocity;
@@ -469,6 +605,71 @@
 
 };
 
+class UChaosVehicleSimulation
+{
+public:
+	virtual ~UChaosVehicleSimulation()
+	{
+		PVehicle.Reset(nullptr);
+	}
 
+	virtual void Init(TUniquePtr<Chaos::FSimpleWheeledVehicle>& PVehicleIn)
+	{
+		PVehicle = MoveTemp(PVehicleIn);
+	}
+
+	virtual void TickVehicle(UWorld* WorldIn, float DeltaTime, const FChaosVehicleDefaultAsyncInput& InputData, FChaosVehicleAsyncOutput& OutputData, Chaos::TPBDRigidParticleHandle<float, 3>* Handle);
+
+	/** Advance the vehicle simulation */
+	virtual void UpdateSimulation(float DeltaTime, const FChaosVehicleDefaultAsyncInput& InputData, Chaos::TPBDRigidParticleHandle<float, 3>* Handle);
+	
+	virtual void FillOutputState(FChaosVehicleAsyncOutput& Output);
+
+	/** Are enough vehicle systems specified such that physics vehicle simulation is possible */
+	virtual bool CanSimulate() const { return true; }
+
+	/** Pass control Input to the vehicle systems */
+	virtual void ApplyInput(const FControlInputs& ControlInputs, float DeltaTime);
+
+	/** Apply aerodynamic forces to vehicle body */
+	virtual void ApplyAerodynamics(float DeltaTime);
+
+	/** Apply Aerofoil forces to vehicle body */
+	virtual void ApplyAerofoilForces(float DeltaTime);
+
+	/** Apply Thruster forces to vehicle body */
+	virtual void ApplyThrustForces(float DeltaTime);
+
+	/** Apply direct control over vehicle body rotation */
+	virtual void ApplyTorqueControl(float DeltaTime, const FChaosVehicleDefaultAsyncInput& InputData);
+
+
+	/** Add a force to this vehicle */
+	void AddForce(const FVector& Force, bool bAllowSubstepping = true, bool bAccelChange = false);
+	/** Add a force at a particular position (world space when bIsLocalForce = false, body space otherwise) */
+	void AddForceAtPosition(const FVector& Force, const FVector& Position, bool bAllowSubstepping = true, bool bIsLocalForce = false);
+	/** Add an impulse to this vehicle */
+	void AddImpulse(const FVector& Impulse, bool bVelChange);
+	/** Add an impulse to this vehicle and a particular world position */
+	void AddImpulseAtPosition(const FVector& Impulse, const FVector& Position);
+	/** Add a torque to this vehicle */
+	void AddTorqueInRadians(const FVector& Torque, bool bAllowSubstepping = true, bool bAccelChange = false);
+
+
+	/** Draw debug text for the wheels and suspension */
+	virtual void DrawDebug3D();
+	UWorld* World;
+
+	// Physics Thread Representation of chassis rigid body
+	Chaos::TPBDRigidParticleHandle<float, 3>* RigidHandle;
+
+	FVehicleState VehicleState;
+
+	// #todo: this isn't very configurable
+	TUniquePtr<Chaos::FSimpleWheeledVehicle> PVehicle;
+
+};
+
+
 /**
  * Base component to handle the vehicle simulation for an actor.
@@ -476,6 +677,8 @@
 UCLASS(Abstract, hidecategories=(PlanarMovement, "Components|Movement|Planar", Activation, "Components|Activation"))
 class CHAOSVEHICLES_API UChaosVehicleMovementComponent : public UPawnMovementComponent
 {
+	friend struct FChaosVehicleDefaultAsyncInput;
+
 	GENERATED_UCLASS_BODY()
 
 //#todo: these 2 oddities seem out of place
@@ -566,6 +769,10 @@
 	UPROPERTY(Transient)
 	uint32 bWasAvoidanceUpdated : 1;
 
+	int32 TargetGearInput;
+	int8 SetGearImmediate;
+	Chaos::ETransmissionType TransmissionType;
+
 public:
 
 	/** UObject interface */
@@ -593,10 +800,6 @@
 	/** Return true if we are ready to create a vehicle, false if the setup has missing references */
 	virtual bool CanCreateVehicle() const;
 
-	/** Are enough vehicle systems specified such that physics vehicle simulation is possible */
-	virtual bool CanSimulate() const { return true; }
-
-
 	/** Used to create any physics engine information for this component */
 	virtual void OnCreatePhysicsState() override;
 
@@ -604,10 +807,7 @@
 	virtual void OnDestroyPhysicsState() override;
 
 	/** Updates the vehicle tuning and other state such as user input. */
-	virtual void PreTick(float DeltaTime);
+	virtual void PreTickGT(float DeltaTime);
 
-	/** Tick this vehicle sim right before input is sent to the vehicle system */
-	virtual void TickVehicle( float DeltaTime );
-
 	/** Stops movement immediately (zeroes velocity, usually zeros acceleration for components with acceleration). */
 	virtual void StopMovementImmediately() override;
@@ -696,11 +896,43 @@
 	/** location local coordinates of named bone in skeleton, apply additional offset or just use offset if no bone located */
 	FVector LocateBoneOffset(const FName InBoneName, const FVector& InExtraOffset) const;
 
-	TUniquePtr<Chaos::FSimpleWheeledVehicle>& PhysicsVehicle()
+	TUniquePtr<FPhysicsVehicleOutput>& PhysicsVehicleOutput()
 	{
-		return PVehicle;
+		return PVehicleOutput;
 	}
 
+	//----ASYNC----
+	TUniquePtr<FChaosVehicleAsyncInput> SetCurrentAsyncInputOutput(int32 InputIdx, FChaosVehicleManagerAsyncOutput* CurOutput, FChaosVehicleManagerAsyncOutput* NextOutput, float Alpha, int32 VehicleManagerTimestamp);
+
+	void SetCurrentAsyncInputOutputInternal(FChaosVehicleAsyncInput* CurInput, int32 InputIdx, FChaosVehicleManagerAsyncOutput* CurOutput, int32 VehicleManagerTimestamp);
+	void SetCurrentAsyncInputOutputInternal(FChaosVehicleAsyncInput* CurInput, int32 InputIdx, FChaosVehicleManagerAsyncOutput* CurOutput, FChaosVehicleManagerAsyncOutput* NextOutput, float Alpha, int32 VehicleManagerTimestamp);
+
+	void Update(float DeltaTime);
+
+	// Get output data from Physics Thread
+	virtual void ParallelUpdate(float DeltaSeconds);
+
+	void FinalizeSimCallbackData(FChaosVehicleManagerAsyncInput& Input);
+
+	EChaosAsyncVehicleDataType CurAsyncType;
+	FChaosVehicleAsyncInput* CurAsyncInput;
+	struct FChaosVehicleAsyncOutput* CurAsyncOutput;
+	struct FChaosVehicleAsyncOutput* NextAsyncOutput;
+	float OutputInterpAlpha;
+
+	struct FAsyncOutputWrapper
+	{
+		int32 Idx;
+		int32 Timestamp;
+
+		FAsyncOutputWrapper()
+			: Idx(INDEX_NONE)
+			, Timestamp(INDEX_NONE)
+		{
+		}
+	};
+	TArray<FAsyncOutputWrapper> OutputsWaitingOn;
+
 protected:
 
 	// replicated state of vehicle 
@@ -802,6 +1034,7 @@
 	UPROPERTY(EditAnywhere, Category = VehicleInput, AdvancedDisplay)
 	FVehicleInputRateConfig YawInputRate;
 
+
 	// input related
 
 	/** Compute steering input */
@@ -850,26 +1083,8 @@
 	/** Read current state for simulation */
 	virtual void UpdateState(float DeltaTime);
 
-	/** Advance the vehicle simulation */
-	virtual void UpdateSimulation(float DeltaTime);
-
-	/** Pass control Input to the vehicle systems */
-	virtual void ApplyInput(float DeltaTime);
-
-	/** Apply aerodynamic forces to vehicle body */
-	virtual void ApplyAerodynamics(float DeltaTime);
-
-	/** Apply Aerofoil forces to vehicle body */
-	virtual void ApplyAerofoilForces(float DeltaTime);
-
-	/** Apply Thruster forces to vehicle body */
-	virtual void ApplyThrustForces(float DeltaTime);
-
-	/** Apply direct control over vehicle body rotation */
-	virtual void ApplyTorqueControl(float DeltaTime);
-
 	/** Option to aggressively sleep the vehicle */
-	virtual void ProcessSleeping();
+	virtual void ProcessSleeping(const FControlInputs& ControlInputs);
 
 	/** Pass current state to server */
 	UFUNCTION(reliable, server, WithValidation)
@@ -894,13 +1109,13 @@
 	/** Create and setup the Chaos vehicle */
 	virtual void CreateVehicle();
 
-	virtual void CreatePhysicsVehicle();
+	virtual TUniquePtr<Chaos::FSimpleWheeledVehicle> CreatePhysicsVehicle();
 
 	/** Skeletal mesh needs some special handling in the vehicle case */
 	virtual void FixupSkeletalMesh() {}
 
 	/** Allocate and setup the Chaos vehicle */
-	virtual void SetupVehicle();
+	virtual void SetupVehicle(TUniquePtr<Chaos::FSimpleWheeledVehicle>& PVehicle);
 
 	/** Do some final setup after the Chaos vehicle gets created */
 	virtual void PostSetupVehicle();
@@ -920,9 +1135,6 @@
 	/** Draw debug text for the wheels and suspension */
 	virtual void DrawDebug(UCanvas* Canvas, float& YL, float& YPos);
 
-	/** Draw debug text for the wheels and suspension */
-	virtual void DrawDebug3D();
-
 	// draw 2D debug line to UI canvas
 	void DrawLine2D(UCanvas* Canvas, const FVector2D& StartPos, const FVector2D& EndPos, FColor Color, float Thickness = 1.f);
 
@@ -933,25 +1145,15 @@
 
 	FBodyInstance* GetBodyInstance();
 
-	FVehicleState VehicleState;
+	FVehicleState VehicleState;							/* Useful vehicle state captured at start of frame */
+	TUniquePtr<FPhysicsVehicleOutput> PVehicleOutput;	/* physics simulation data output from the async physics thread */
 
-	// #todo: this isn't very configurable
-	TUniquePtr<Chaos::FSimpleWheeledVehicle> PVehicle;
-
 	/** Handle for delegate registered on mesh component */
 	FDelegateHandle MeshOnPhysicsStateChangeHandle;
 
 protected:
-	/** Add a force to this vehicle */
-	void AddForce(const FVector& Force, bool bAllowSubstepping = true, bool bAccelChange = false);
-	/** Add a force at a particular position (world space when bIsLocalForce = false, body space otherwise) */
-	void AddForceAtPosition(const FVector& Force, const FVector& Position, bool bAllowSubstepping = true, bool bIsLocalForce = false);
-	/** Add an impulse to this vehicle */
-	void AddImpulse(const FVector& Impulse, bool bVelChange);
-	/** Add an impulse to this vehicle and a particular world position */
-	void AddImpulseAtPosition(const FVector& Impulse, const FVector& Position);
-	/** Add a torque to this vehicle */
-	void AddTorqueInRadians(const FVector& Torque, bool bAllowSubstepping = true, bool bAccelChange = false);
 
+	TUniquePtr<UChaosVehicleSimulation> VehicleSimulationPT;	/* simulation code running on the physics thread async callback */
+
 private:
 	UPROPERTY(transient, Replicated)
diff --git a/Engine/Plugins/Experimental/ChaosVehiclesPlugin/Source/ChaosVehicles/Public/ChaosVehicleWheel.h b/Engine/Plugins/Experimental/ChaosVehiclesPlugin/Source/ChaosVehicles/Public/ChaosVehicleWheel.h
--- a/Engine/Plugins/Experimental/ChaosVehiclesPlugin/Source/ChaosVehicles/Public/ChaosVehicleWheel.h
+++ b/Engine/Plugins/Experimental/ChaosVehiclesPlugin/Source/ChaosVehicles/Public/ChaosVehicleWheel.h
@@ -236,7 +236,7 @@ class UChaosWheeledVehicleMovementComponent;
 
 		/** The vehicle that owns us */
 		UPROPERTY(transient)
-		class UChaosWheeledVehicleMovementComponent* VehicleSim;
+		class UChaosWheeledVehicleMovementComponent* VehicleComponent;
 
 		// Our index in the vehicle's (and setup's) wheels array
 		UPROPERTY(transient)
@@ -354,6 +354,7 @@ class UChaosWheeledVehicleMovementComponent;
 			PWheelConfig.EngineEnabled = this->bAffectedByEngine;
 			PWheelConfig.ABSEnabled = this->bABSEnabled;
 			PWheelConfig.TractionControlEnabled = this->bTractionControlEnabled;
+			PWheelConfig.AxleType = static_cast<Chaos::FSimpleWheelConfig::EAxleType>(this->AxleType);
 			PWheelConfig.LongitudinalFrictionMultiplier = this->LongitudinalFrictionForceMultiplier;
 			PWheelConfig.LateralFrictionMultiplier = this->LateralFrictionForceMultiplier;
 			PWheelConfig.SideSlipModifier = this->SideSlipModifier;
diff --git a/Engine/Plugins/Experimental/ChaosVehiclesPlugin/Source/ChaosVehicles/Public/ChaosWheeledVehicleMovementComponent.h b/Engine/Plugins/Experimental/ChaosVehiclesPlugin/Source/ChaosVehicles/Public/ChaosWheeledVehicleMovementComponent.h
--- a/Engine/Plugins/Experimental/ChaosVehiclesPlugin/Source/ChaosVehicles/Public/ChaosWheeledVehicleMovementComponent.h
+++ b/Engine/Plugins/Experimental/ChaosVehiclesPlugin/Source/ChaosVehicles/Public/ChaosWheeledVehicleMovementComponent.h
@@ -61,6 +61,8 @@ enum class EVehicleDifferential : uint8
 	RearWheelDrive,
 };
 
+
+
 /**
  * Structure containing information about the status of a single wheel of the vehicle.
  */
@@ -173,11 +175,26 @@ struct FVehicleDifferentialConfig
 	//UPROPERTY(EditAnywhere, Category = Setup, meta = (ClampMin = "1.0", UIMin = "1.0"))
 	//float RearBias;
 
+	const Chaos::FSimpleDifferentialConfig& GetPhysicsDifferentialConfig()
+	{
+		FillDifferentialSetup();
+		return PDifferentialConfig;
+	}
+
 	void InitDefaults()
 	{
 		DifferentialType = EVehicleDifferential::RearWheelDrive;
 		FrontRearSplit = 0.5f;
 	}
+
+	void FillDifferentialSetup()
+	{
+		PDifferentialConfig.DifferentialType = static_cast<Chaos::EDifferentialType>(this->DifferentialType);
+		PDifferentialConfig.FrontRearSplit = this->FrontRearSplit;
+	}
+
+	Chaos::FSimpleDifferentialConfig PDifferentialConfig;
+
 };
 
 USTRUCT()
@@ -229,6 +246,14 @@ struct FVehicleEngineConfig
 		EngineRevDownRate = 600.0f;
 	}
 
+	float GetTorqueFromRPM(float EngineRPM)
+	{
+		// The source curve does not need to be normalized, however we are normalizing it when it is passed on,
+		// since it's the MaxRPM and MaxTorque values that determine the range of RPM and Torque
+		float MinVal = 0.f, MaxVal = 0.f;
+		this->TorqueCurve.GetRichCurveConst()->GetValueRange(MinVal, MaxVal);
+		return TorqueCurve.GetRichCurve()->Eval(EngineRPM) / MaxVal * MaxTorque;
+	}
 private:
 
 	void FillEngineSetup()
@@ -329,6 +354,23 @@ struct FVehicleTransmissionConfig
 		TransmissionEfficiency = 0.9f;
 	}
 
+	float GetGearRatio(int32 InGear)
+	{
+		if (InGear > 0) // a forwards gear
+		{
+			return ForwardGearRatios[InGear - 1] * FinalRatio;
+		}
+		else if (InGear < 0) // a reverse gear
+		{
+			return -ReverseGearRatios[FMath::Abs(InGear) - 1] * FinalRatio;
+		}
+		else
+		{
+			return 0.f; // neutral has no ratio
+		}
+	}
+
+
 private:
 
 	void FillTransmissionSetup()
@@ -472,6 +514,8 @@ struct CHAOSVEHICLES_API FWheelState
 
 	/** Commonly used Wheel state - evaluated once used wherever required for that frame */
 	void CaptureState(int WheelIdx, const FVector& WheelOffset, const FBodyInstance* TargetInstance);
+	void CaptureState(int WheelIdx, const FVector& WheelOffset, const Chaos::TPBDRigidParticleHandle<float, 3>* Handle);
+	FVector GetVelocityAtPoint(const Chaos::TPBDRigidParticleHandle<float, 3>* Rigid, const FVector& InPoint);
 
 	TArray<FVector> WheelWorldLocation;	/** Current Location Of Wheels In World Coordinates */
 	TArray<FVector> WorldWheelVelocity; /** Current velocity at wheel location In World Coordinates - combined linear and angular */
@@ -479,6 +523,76 @@ struct CHAOSVEHICLES_API FWheelState
 	TArray<Chaos::FSuspensionTrace> Trace;
 };
 
+//////////////////////////////////////////////////////////////////////////
+
+class UChaosWheeledVehicleSimulation : public UChaosVehicleSimulation
+{
+public:
+
+	UChaosWheeledVehicleSimulation(TArray<class UChaosVehicleWheel*>& WheelsIn, TArray<FPhysicsConstraintHandle>& ConstraintHandlesIn)
+		: Wheels(WheelsIn)
+		, ConstraintHandles(ConstraintHandlesIn)
+	{
+
+	}
+
+	virtual ~UChaosWheeledVehicleSimulation()
+	{
+	}
+
+	virtual void Init(TUniquePtr<Chaos::FSimpleWheeledVehicle>& PVehicleIn) override
+	{
+		UChaosVehicleSimulation::Init(PVehicleIn);
+
+		WheelState.Init(PVehicle->Wheels.Num());
+	}
+
+	virtual void TickVehicle(UWorld* WorldIn, float DeltaTime, const FChaosVehicleDefaultAsyncInput& InputData, FChaosVehicleAsyncOutput& OutputData, Chaos::TPBDRigidParticleHandle<float, 3>* Handle) override;
+
+	/** Advance the vehicle simulation */
+	virtual void UpdateSimulation(float DeltaTime, const FChaosVehicleDefaultAsyncInput& InputData, Chaos::TPBDRigidParticleHandle<float, 3>* Handle) override;
+
+	virtual void FillOutputState(FChaosVehicleAsyncOutput& Output) override;
+
+	/** Are enough vehicle systems specified such that physics vehicle simulation is possible */
+	virtual bool CanSimulate() const override;
+
+	/** Pass control Input to the vehicle systems */
+	virtual void ApplyInput(const FControlInputs& ControlInputs, float DeltaTime) override;
+
+	/** Perform suspension ray/shape traces */
+	virtual void PerformSuspensionTraces(const TArray<Chaos::FSuspensionTrace>& SuspensionTrace, const FCollisionQueryParams& TraceParams);
+
+
+	/** Update the engine/transmission simulation */
+	virtual void ProcessMechanicalSimulation(float DeltaTime);
+
+	/** Process steering mechanism */
+	virtual void ProcessSteering(const FControlInputs& ControlInputs);
+
+	/** calculate and apply lateral and longitudinal friction forces from wheels */
+	virtual void ApplyWheelFrictionForces(float DeltaTime);
+
+	/** calculate and apply chassis suspension forces */
+	virtual void ApplySuspensionForces(float DeltaTime);
+
+	bool IsWheelSpinning() const;
+
+
+	/** Draw 3D debug lines and things along side the 3D model */
+	virtual void DrawDebug3D() override;
+
+	// #TODO - should not exist here, now duplicated, best not to access at all from Physics 
+	TArray<class UChaosVehicleWheel*>& Wheels;
+
+	FWheelState WheelState;	/** Cached state that holds wheel data for this frame */
+
+	TArray<FPhysicsConstraintHandle>& ConstraintHandles;
+
+};
+
+//////////////////////////////////////////////////////////////////////////
+
 UCLASS(ClassGroup = (Physics), meta = (BlueprintSpawnableComponent), hidecategories = (PlanarMovement, "Components|Movement|Planar", Activation, "Components|Activation"))
 class CHAOSVEHICLES_API UChaosWheeledVehicleMovementComponent : public UChaosVehicleMovementComponent
 {
@@ -539,16 +653,6 @@ class CHAOSVEHICLES_API UChaosWheeledVehicleMovementComponent : public UChaosVeh
 	static FWheelStatus MakeWheelStatus(bool bInContact, FVector& ContactPoint, UPhysicalMaterial* PhysMaterial
 			, float NormalizedSuspensionLength, float SpringForce, bool bIsSlipping, float SlipMagnitude, bool bIsSkidding, float SkidMagnitude, FVector& SkidNormal);
 
-
-
-
-
-
-
-
-
-
-
 	/** Get a wheels current simulation state */
 	UFUNCTION(BlueprintCallable, Category = "Game|Components|ChaosWheeledVehicleMovement")
 	const FWheelStatus& GetWheelState(int WheelIndex) const
@@ -561,24 +665,21 @@ class CHAOSVEHICLES_API UChaosWheeledVehicleMovementComponent : public UChaosVeh
 
 	virtual void Serialize(FArchive & Ar) override;
 
+	// Get output data from Physics Thread
+	virtual void ParallelUpdate(float DeltaSeconds);
+
 #if WITH_EDITOR
 	virtual void PostEditChangeProperty(struct FPropertyChangedEvent& PropertyChangedEvent) override;
 #endif
 	/** Are the configuration references configured sufficiently that the vehicle can be created */
 	virtual bool CanCreateVehicle() const override;
 
-	/** Are the appropriate vehicle systems specified such that physics vehicle simulation is possible */
-	virtual bool CanSimulate() const override;
-
 	/** Used to create any physics engine information for this component */
 	virtual void OnCreatePhysicsState() override;
 
 	/** Used to shut down and pysics engine structure for this component */
 	virtual void OnDestroyPhysicsState() override;
 
-	/** Tick this vehicle sim right before input is sent to the vehicle system  */
-	virtual void TickVehicle(float DeltaTime) override;
-
 	/** display next debug page */
 	static void NextDebugPage();
 
@@ -612,6 +713,15 @@ protected:
 	//////////////////////////////////////////////////////////////////////////
 	// Setup
 
+	/** */
+	virtual TUniquePtr<Chaos::FSimpleWheeledVehicle> CreatePhysicsVehicle() override
+	{
+		// Make the Vehicle Simulation class that will be updated from the physics thread async callback
+		VehicleSimulationPT = MakeUnique<UChaosWheeledVehicleSimulation>(Wheels, ConstraintHandles);
+
+		return UChaosVehicleMovementComponent::CreatePhysicsVehicle();
+	}
+
 	/** Re-Compute any runtime constants values that rely on setup data */
 	virtual void ComputeConstants() override;
 
@@ -619,7 +729,10 @@ protected:
 	virtual void FixupSkeletalMesh();
 
 	/** Allocate and setup the Chaos vehicle */
-	virtual void SetupVehicle() override;
+	virtual void SetupVehicle(TUniquePtr<Chaos::FSimpleWheeledVehicle>& PVehicle) override;
+
+	/** Create and setup the Chaos vehicle */
+	virtual void CreateVehicle();
 
 	/** Instantiate and setup our wheel objects */
 	virtual void CreateWheels();
@@ -631,7 +744,7 @@ protected:
 	virtual void SetupVehicleShapes();
 
 	/** Setup calculated suspension parameters */
-	void SetupSuspension();
+	void SetupSuspension(TUniquePtr<Chaos::FSimpleWheeledVehicle>& PVehicle);
 
 	/** Maps UChaosVehicleWheel Axle to a wheel index */
 	void RecalculateAxles();
@@ -641,27 +754,7 @@ protected:
 
 	//////////////////////////////////////////////////////////////////////////
 	// Update
-
-	/** Advance the vehicle simulation */
-	virtual void UpdateSimulation(float DeltaTime) override;
-
-	/** Perform suspension ray/shape traces */
-	virtual void PerformSuspensionTraces(const TArray<Chaos::FSuspensionTrace>& SuspensionTrace);
-
-	/** Pass control Input to the vehicle systems */
-	virtual void ApplyInput(float DeltaTime);
-
-	/** Update the engine/transmission simulation */
-	virtual void ProcessMechanicalSimulation(float DeltaTime);
-
-	/** Process steering mechanism */
-	virtual void ProcessSteering();
-
-	/** calculate and apply lateral and longitudinal friction forces from wheels */
-	virtual void ApplyWheelFrictionForces(float DeltaTime);
-
-	/** calculate and apply chassis suspension forces */
-	virtual void ApplySuspensionForces(float DeltaTime);
+	void FillWheelOutputState();
 
 	//////////////////////////////////////////////////////////////////////////
 	// Debug
@@ -669,9 +762,6 @@ protected:
 	/** Draw 2D debug text graphs on UI for the wheels, suspension and other systems */
 	virtual void DrawDebug(UCanvas* Canvas, float& YL, float& YPos);
 
-	/** Draw 3D debug lines and things along side the 3D model */
-	virtual void DrawDebug3D() override;
-
 	/** Get distances between wheels - primarily a debug display helper */
 	const FVector2D& GetWheelLayoutDimensions() const
 	{
@@ -680,11 +770,9 @@ protected:
 
 	private:
 
-	void FillWheelOutputState();
 
 	/** Get distances between wheels - primarily a debug display helper */
 	FVector2D CalculateWheelLayoutDimensions();
-	bool IsWheelSpinning() const;
 
 #if !(UE_BUILD_SHIPPING || UE_BUILD_TEST)
 	float CalcDialAngle(float CurrentValue, float MaxValue);
@@ -693,7 +781,6 @@ protected:
 
 	static EDebugPages DebugPage;
 	uint32 NumDrivenWheels; /** The number of wheels that have engine enabled checked */
-	FWheelState WheelState;	/** Cached state that hold wheel data for this frame */
 	FVector2D WheelTrackDimensions;	// Wheelbase (X) and track (Y) dimensions
 	TMap<UChaosVehicleWheel*, TArray<int>> AxleToWheelMap;
 	TArray<FPhysicsConstraintHandle> ConstraintHandles;
diff --git a/Engine/Source/Runtime/Experimental/ChaosVehicles/ChaosVehiclesCore/Private/ArcadeSystem.cpp b/Engine/Source/Runtime/Experimental/ChaosVehicles/ChaosVehiclesCore/Private/ArcadeSystem.cpp
new file mode 100644
--- /dev/null
+++ b/Engine/Source/Runtime/Experimental/ChaosVehicles/ChaosVehiclesCore/Private/ArcadeSystem.cpp
@@ -0,0 +1,34 @@
+// Copyright Epic Games, Inc. All Rights Reserved.
+
+#include "ArcadeSystem.h"
+
+#if VEHICLE_DEBUGGING_ENABLED
+PRAGMA_DISABLE_OPTIMIZATION
+#endif
+
+namespace Chaos
+{
+
+	FTorqueControlSim::FTorqueControlSim(const FTorqueControlConfig* SetupIn) 
+		: TVehicleSystem<FTorqueControlConfig>(SetupIn)
+	{
+
+	}
+
+	FTargetRotationControlSim::FTargetRotationControlSim(const FTargetRotationControlConfig* SetupIn) 
+		: TVehicleSystem<FTargetRotationControlConfig>(SetupIn)
+	{
+
+	}
+
+	FStabilizeControlSim::FStabilizeControlSim(const FStabilizeControlConfig* SetupIn) 
+		: TVehicleSystem<FStabilizeControlConfig>(SetupIn)
+	{
+
+	}
+
+} // namespace Chaos
+
+#if VEHICLE_DEBUGGING_ENABLED
+PRAGMA_ENABLE_OPTIMIZATION
+#endif
diff --git a/Engine/Source/Runtime/Experimental/ChaosVehicles/ChaosVehiclesCore/Private/WheelSystem.cpp b/Engine/Source/Runtime/Experimental/ChaosVehicles/ChaosVehiclesCore/Private/WheelSystem.cpp
--- a/Engine/Source/Runtime/Experimental/ChaosVehicles/ChaosVehiclesCore/Private/WheelSystem.cpp
+++ b/Engine/Source/Runtime/Experimental/ChaosVehicles/ChaosVehiclesCore/Private/WheelSystem.cpp
@@ -186,6 +186,15 @@ namespace Chaos
 		return;
 	}
 
+
+
+
+	FAxleSim::FAxleSim() : TVehicleSystem<FAxleConfig>(&Setup)
+	{
+	}
+
+
+
 } // namespace Chaos
 
 #if VEHICLE_DEBUGGING_ENABLED
diff --git a/Engine/Source/Runtime/Experimental/ChaosVehicles/ChaosVehiclesCore/Public/ArcadeSystem.h b/Engine/Source/Runtime/Experimental/ChaosVehicles/ChaosVehiclesCore/Public/ArcadeSystem.h
new file mode 100644
--- /dev/null
+++ b/Engine/Source/Runtime/Experimental/ChaosVehicles/ChaosVehiclesCore/Public/ArcadeSystem.h
@@ -0,0 +1,170 @@
+// Copyright Epic Games, Inc. All Rights Reserved.
+
+#pragma once
+
+#include "VehicleSystemTemplate.h"
+#include "VehicleUtility.h"
+
+#if VEHICLE_DEBUGGING_ENABLED
+PRAGMA_DISABLE_OPTIMIZATION
+#endif
+
+namespace Chaos
+{
+
+	struct CHAOSVEHICLESCORE_API FTorqueControlConfig
+	{
+		/** Torque Control Enabled */
+		bool Enabled;
+
+		/** Yaw Torque Scaling */
+		float YawTorqueScaling;
+
+		float YawFromSteering;
+
+		float YawFromRollTorqueScaling;
+
+		/** Pitch Torque Scaling */
+		float PitchTorqueScaling;
+
+		/** Roll Torque Scaling */
+		float RollTorqueScaling;
+
+		float RollFromSteering;
+
+		/** Rotation damping */
+		float RotationDamping;
+
+		void InitDefaults()
+		{
+			Enabled = false;
+			YawTorqueScaling = 0.0f;
+			YawFromSteering = 0.0f;
+			YawFromRollTorqueScaling = 0.0f;
+			PitchTorqueScaling = 0.0f;
+			RollTorqueScaling = 0.0f;
+			RollFromSteering = 0.0f;
+			RotationDamping = 0.02f;
+		}
+	};
+
+	struct CHAOSVEHICLESCORE_API FTargetRotationControlConfig
+	{
+		/** Rotation Control Enabled */
+		bool Enabled;
+
+		bool bRollVsSpeedEnabled;
+
+		float RollControlScaling;
+
+		float RollMaxAngle;
+
+		float PitchControlScaling;
+
+		float PitchMaxAngle;
+
+		/** Rotation stiffness */
+		float RotationStiffness;
+
+		/** Rotation damping */
+		float RotationDamping;
+
+		/** Rotation mac accel */
+		float MaxAccel;
+
+		float AutoCentreRollStrength;
+
+		float AutoCentrePitchStrength;
+
+		float AutoCentreYawStrength;
+
+
+		void InitDefaults()
+		{
+			Enabled = false;
+
+			bRollVsSpeedEnabled = false;
+
+			RollControlScaling = 0.f;
+			RollMaxAngle = 0.f;
+			PitchControlScaling = 0.f;
+			PitchMaxAngle = 0.f;
+
+			RotationStiffness = 0.f;
+			RotationDamping = 0.2;
+			MaxAccel = 0.f;
+
+			AutoCentreRollStrength = 0.f;
+			AutoCentrePitchStrength = 0.f;
+			AutoCentreYawStrength = 0.f;
+		}
+	};
+
+	struct CHAOSVEHICLESCORE_API FStabilizeControlConfig
+	{
+
+		/** Torque Control Enabled */
+		bool Enabled;
+
+		/** Yaw Torque Scaling */
+		float AltitudeHoldZ;
+
+		float PositionHoldXY;
+
+		void InitDefaults()
+		{
+			Enabled = false;
+			AltitudeHoldZ = 4.0f;
+			PositionHoldXY = 8.0f;
+		}
+	};
+
+	struct CHAOSVEHICLESCORE_API FArcadeControlConfig
+	{
+		FTorqueControlConfig TorqueControl;
+
+		/** Arcade style direct control of vehicle rotation via torque force */
+		FTargetRotationControlConfig TargetRotationControl;
+
+		/** Arcade style control of vehicle */
+		FStabilizeControlConfig StabilizeControl;
+	};
+
+
+	class CHAOSVEHICLESCORE_API FTorqueControlSim : public TVehicleSystem<FTorqueControlConfig>
+	{
+	public:
+		FTorqueControlSim()
+		{}
+
+		FTorqueControlSim(const FTorqueControlConfig* SetupIn);
+
+	};
+
+
+	class CHAOSVEHICLESCORE_API FTargetRotationControlSim : public TVehicleSystem<FTargetRotationControlConfig>
+	{
+	public:
+		FTargetRotationControlSim()
+		{}
+
+		FTargetRotationControlSim(const FTargetRotationControlConfig* SetupIn);
+
+	};
+
+
+	class CHAOSVEHICLESCORE_API FStabilizeControlSim : public TVehicleSystem<FStabilizeControlConfig>
+	{
+	public:
+		FStabilizeControlSim()
+		{}
+
+		FStabilizeControlSim(const FStabilizeControlConfig* SetupIn);
+	};
+
+
+} // namespace Chaos
+
+#if VEHICLE_DEBUGGING_ENABLED
+PRAGMA_ENABLE_OPTIMIZATION
+#endif
\ No newline at end of file
diff --git a/Engine/Source/Runtime/Experimental/ChaosVehicles/ChaosVehiclesCore/Public/SimpleVehicle.h b/Engine/Source/Runtime/Experimental/ChaosVehicles/ChaosVehiclesCore/Public/SimpleVehicle.h
--- a/Engine/Source/Runtime/Experimental/ChaosVehicles/ChaosVehiclesCore/Public/SimpleVehicle.h
+++ b/Engine/Source/Runtime/Experimental/ChaosVehicles/ChaosVehiclesCore/Public/SimpleVehicle.h
@@ -11,6 +11,7 @@
 #include "AerodynamicsSystem.h"
 #include "AerofoilSystem.h"
 #include "ThrustSystem.h"
+#include "ArcadeSystem.h"
 //#include "FuelSystem.h"
 //#include "BallastSystem.h"
 
@@ -35,6 +36,10 @@ class FSimpleWheeledVehicle : public IVehicleInterface
 {
 public:
 	FSimpleWheeledVehicle()
+	: bSuspensionEnabled(true)
+	, bMechanicalSimEnabled(true)
+	, bWheelFrictionEnabled(true)
+	, NumDrivenWheels(0)
 	{
 
 	}
@@ -69,12 +74,34 @@ public:
 		return (Transmission.Num() > 0);
 	}
 
+	bool HasTorqueControlSetup()
+	{
+		return (TorqueControlSim.Num() > 0);
+	}
+
+	bool HasTargetRotationControlSetup()
+	{
+		return (TargetRotationControlSim.Num() > 0);
+	}
+
+	bool HasStabilizeControlSetup()
+	{
+		return (StabilizeControlSim.Num() > 0);
+	}
+
+
 	FSimpleTransmissionSim& GetTransmission()
 	{
 		check(Transmission.Num() == 1);
 		return Transmission[0];
 	}
 
+	FSimpleDifferentialSim& GetDifferential()
+	{
+		check(Differential.Num() == 1);
+		return Differential[0];
+	}
+
 	FSimpleWheelSim& GetWheel(int WheelIdx)
 	{
 		check(WheelIdx < Wheels.Num());
@@ -111,18 +138,53 @@ public:
 		return Thrusters[ThrusterIdx];
 	}
 
+	FTorqueControlSim& GetTorqueControl()
+	{
+		check(TorqueControlSim.Num() == 1);
+		return TorqueControlSim[0];
+	}
+
+	FTargetRotationControlSim& GetTargetRotationControl()
+	{
+		check(TargetRotationControlSim.Num() == 1);
+		return TargetRotationControlSim[0];
+	}
+
+	FStabilizeControlSim& GetStabilizeControl()
+	{
+		check(StabilizeControlSim.Num() == 1);
+		return StabilizeControlSim[0];
+	}
+
+	const TArray<FAxleSim>& GetAxles() const
+	{
+		return Axles;
+	}
+
+
 	TArray<FSimpleEngineSim> Engine;
 	TArray<FSimpleTransmissionSim> Transmission;
+	TArray<FSimpleDifferentialSim> Differential;
 	TArray<FSimpleWheelSim> Wheels;
 	TArray<FSimpleSuspensionSim> Suspension;
 	TArray<FSimpleSteeringSim> Steering;
 	TArray<FSimpleAerodynamicsSim> Aerodynamics;
 	TArray<FAerofoil> Aerofoils;
 	TArray<FSimpleThrustSim> Thrusters;
+	TArray<FAxleSim> Axles;
+
+	TArray<FTorqueControlSim> TorqueControlSim;
+	TArray<FTargetRotationControlSim> TargetRotationControlSim;
+	TArray<FStabilizeControlSim> StabilizeControlSim;
 	//TArray<FSimpleFuelSim> Fuel;
 	//TArray<FSimpleBallastSim> Ballast;
 	// turbo
 	// .. 
+
+	bool bSuspensionEnabled;
+	bool bMechanicalSimEnabled;
+	bool bWheelFrictionEnabled;
+	uint32 NumDrivenWheels;
 };
 
 
diff --git a/Engine/Source/Runtime/Experimental/ChaosVehicles/ChaosVehiclesCore/Public/SystemTemplate.h b/Engine/Source/Runtime/Experimental/ChaosVehicles/ChaosVehiclesCore/Public/SystemTemplate.h
deleted file mode 100644
diff --git a/Engine/Source/Runtime/Experimental/ChaosVehicles/ChaosVehiclesCore/Public/TransmissionSystem.h b/Engine/Source/Runtime/Experimental/ChaosVehicles/ChaosVehiclesCore/Public/TransmissionSystem.h
--- a/Engine/Source/Runtime/Experimental/ChaosVehicles/ChaosVehiclesCore/Public/TransmissionSystem.h
+++ b/Engine/Source/Runtime/Experimental/ChaosVehicles/ChaosVehiclesCore/Public/TransmissionSystem.h
@@ -25,6 +25,33 @@ namespace Chaos
 		Automatic
 	};
 
+	enum EDifferentialType : uint8
+	{
+		AllWheelDrive,
+		FrontWheelDrive,
+		RearWheelDrive,
+	};
+
+	struct CHAOSVEHICLESCORE_API FSimpleDifferentialConfig
+	{
+		FSimpleDifferentialConfig()
+			: DifferentialType(EDifferentialType::RearWheelDrive)
+			, FrontRearSplit(0.5f)
+		{
+		}
+
+		EDifferentialType DifferentialType;
+		float FrontRearSplit;
+	};
+
+	class CHAOSVEHICLESCORE_API FSimpleDifferentialSim : public TVehicleSystem<FSimpleDifferentialConfig>
+	{
+	public:
+		FSimpleDifferentialSim(const FSimpleDifferentialConfig* SetupIn) 
+			: TVehicleSystem<FSimpleDifferentialConfig>(SetupIn) {}
+
+	};
+
 	struct CHAOSVEHICLESCORE_API FSimpleTransmissionConfig
 	{
 		FSimpleTransmissionConfig()
@@ -49,6 +76,7 @@ namespace Chaos
 		float TransmissionEfficiency;	// Loss from friction in the system mean we might run at around 0.94 Efficiency
 
 		ETransmissionType TransmissionType;	// Specify Automatic or Manual transmission
+
 		bool AutoReverse;					// Arcade handling - holding Brake switches into reverse after vehicle has stopped
 	};
 
@@ -121,6 +149,11 @@ namespace Chaos
 		/** Get the transmission RPM, from the specified engine RPM and gear selection */
 		float GetTransmissionRPM(float InEngineRPM, int InGear)
 		{
+			if (InGear == 0) // neutral, don't want to divide by zero
+			{
+				return 0.0f;
+			}
+
 			return InEngineRPM / GetGearRatio(InGear);
 		}
 
diff --git a/Engine/Source/Runtime/Experimental/ChaosVehicles/ChaosVehiclesCore/Public/VehicleUtility.h b/Engine/Source/Runtime/Experimental/ChaosVehicles/ChaosVehiclesCore/Public/VehicleUtility.h
--- a/Engine/Source/Runtime/Experimental/ChaosVehicles/ChaosVehiclesCore/Public/VehicleUtility.h
+++ b/Engine/Source/Runtime/Experimental/ChaosVehicles/ChaosVehiclesCore/Public/VehicleUtility.h
@@ -5,7 +5,7 @@
 #include "CoreMinimal.h"
 
 // Disable Optimizations in non debug build configurations
-#define VEHICLE_DEBUGGING_ENABLED 1
+#define VEHICLE_DEBUGGING_ENABLED 0
 
 namespace Chaos
 {
diff --git a/Engine/Source/Runtime/Experimental/ChaosVehicles/ChaosVehiclesCore/Public/WheelSystem.h b/Engine/Source/Runtime/Experimental/ChaosVehicles/ChaosVehiclesCore/Public/WheelSystem.h
--- a/Engine/Source/Runtime/Experimental/ChaosVehicles/ChaosVehiclesCore/Public/WheelSystem.h
+++ b/Engine/Source/Runtime/Experimental/ChaosVehicles/ChaosVehiclesCore/Public/WheelSystem.h
@@ -61,6 +61,13 @@ struct CHAOSVEHICLESCORE_API FSimpleWheelConfig
 		Multiply, // default - most correct
 		Average
 	};
+
+	enum EAxleType : uint8
+	{
+		Undefined = 0,
+		Front,
+		Rear
+	};
 	
 	FSimpleWheelConfig() 
 		: Offset(FVector(2.f, 1.f, 0.f))
@@ -76,6 +83,7 @@ struct CHAOSVEHICLESCORE_API FSimpleWheelConfig
 		, SteeringEnabled(true)
 		, EngineEnabled(false)
 		, TractionControlEnabled(false)
+		, AxleType(EAxleType::Undefined)
 		, FrictionCombineMethod(EFrictionCombineMethod::Multiply)
 		, LateralFrictionMultiplier(1.0f)
 		, LongitudinalFrictionMultiplier(1.0f)
@@ -109,6 +117,8 @@ struct CHAOSVEHICLESCORE_API FSimpleWheelConfig
 	bool EngineEnabled;			// Wheel is driven by an engine
 	bool TractionControlEnabled;// Straight Line Traction Control
 
+	EAxleType AxleType;
+
 	EFrictionCombineMethod FrictionCombineMethod; //#todo: use this variable
 
 	float LateralFrictionMultiplier;
@@ -394,6 +404,24 @@ public:
 };
 
 
+struct CHAOSVEHICLESCORE_API FAxleConfig
+{
+	TArray<uint16> WheelIndex;	// reference to wheels on this axle
+	float RollbarScaling;
+};
+
+class CHAOSVEHICLESCORE_API FAxleSim : public TVehicleSystem<FAxleConfig>
+{
+public:
+
+	FAxleSim();
+
+	void Simulate(float DeltaTime) {}
+
+	FAxleConfig Setup;
+
+};
+
 } // namespace Chaos
 
 #if VEHICLE_DEBUGGING_ENABLED
