ChaosVehicles: Remove logging

diff --git a/Engine/Plugins/Experimental/ChaosVehiclesPlugin/Source/ChaosVehicles/Private/AsyncCallback.cpp b/Engine/Plugins/Experimental/ChaosVehiclesPlugin/Source/ChaosVehicles/Private/AsyncCallback.cpp
--- a/Engine/Plugins/Experimental/ChaosVehiclesPlugin/Source/ChaosVehicles/Private/AsyncCallback.cpp
+++ b/Engine/Plugins/Experimental/ChaosVehiclesPlugin/Source/ChaosVehicles/Private/AsyncCallback.cpp
@@ -60,7 +60,6 @@ void FChaosVehicleManagerAsyncCallback::OnPreSimulate_Internal()
 		auto Handle = VehicleInput.Actor.Proxy->GetHandle();
 		if (Handle->ObjectState() != Chaos::EObjectStateType::Dynamic)
 		{
-			UE_LOG(LogChaos, Warning, TEXT("Handle->ObjectState() %d"), Handle->ObjectState());
 			return;
 		}
 
