ChaosVehicles: Adding a parking mode to hold vehicle in place, so the handbrake input doesn't need to be set every frame.

diff --git a/Engine/Plugins/Experimental/ChaosVehiclesPlugin/Source/ChaosVehicles/Private/ChaosVehicleMovementComponent.cpp b/Engine/Plugins/Experimental/ChaosVehiclesPlugin/Source/ChaosVehicles/Private/ChaosVehicleMovementComponent.cpp
--- a/Engine/Plugins/Experimental/ChaosVehiclesPlugin/Source/ChaosVehicles/Private/ChaosVehicleMovementComponent.cpp
+++ b/Engine/Plugins/Experimental/ChaosVehiclesPlugin/Source/ChaosVehicles/Private/ChaosVehicleMovementComponent.cpp
@@ -591,6 +591,7 @@ UChaosVehicleMovementComponent::UChaosVehicleMovementComponent(const FObjectInit
 	: Super(ObjectInitializer)
 {
 	bReverseAsBrake = true;
+	bParkEnabled = false;
 	Mass = 1500.0f;
 	ChassisWidth = 180.f;
 	ChassisHeight = 140.f;
@@ -793,6 +794,7 @@ void UChaosVehicleMovementComponent::PreTickGT(float DeltaTime)
 		ControlInputs.RollInput = RollInput;
 		ControlInputs.PitchInput = PitchInput;
 		ControlInputs.YawInput = YawInput;
+		ControlInputs.ParkingEnabled = bParkEnabled;
 		ProcessSleeping(ControlInputs);
 	}
 
@@ -864,6 +866,12 @@ void UChaosVehicleMovementComponent::SetHandbrakeInput(bool bNewHandbrake)
 	bRawHandbrakeInput = bNewHandbrake;
 }
 
+
+void UChaosVehicleMovementComponent::SetParked(bool bParked)
+{
+	bParkEnabled = bParked;
+}
+
 void UChaosVehicleMovementComponent::SetChangeUpInput(bool bNewGearUp)
 {
 	bRawGearUpInput = bNewGearUp;
@@ -915,6 +923,10 @@ float UChaosVehicleMovementComponent::GetForwardSpeedMPH() const
 	return Chaos::CmSToMPH(GetForwardSpeed());
 }
 
+bool UChaosVehicleMovementComponent::IsParked() const
+{
+	return (bParkEnabled > 0);
+}
 
 // input related
 float UChaosVehicleMovementComponent::CalcSteeringInput()
@@ -1642,6 +1654,8 @@ void UChaosVehicleMovementComponent::Update(float DeltaTime)
 				AsyncInput->ControlInputs.SetGearImmediate = SetGearImmediate;
 				AsyncInput->ControlInputs.TransmissionType = TransmissionType;
 
+				AsyncInput->ControlInputs.ParkingEnabled = bParkEnabled;
+
 				AsyncInput->GravityZ = GetGravityZ();
 
 				// #TODO: don't create every frame, setup never changes dynamically
diff --git a/Engine/Plugins/Experimental/ChaosVehiclesPlugin/Source/ChaosVehicles/Private/ChaosWheeledVehicleMovementComponent.cpp b/Engine/Plugins/Experimental/ChaosVehiclesPlugin/Source/ChaosVehicles/Private/ChaosWheeledVehicleMovementComponent.cpp
--- a/Engine/Plugins/Experimental/ChaosVehiclesPlugin/Source/ChaosVehicles/Private/ChaosWheeledVehicleMovementComponent.cpp
+++ b/Engine/Plugins/Experimental/ChaosVehiclesPlugin/Source/ChaosVehicles/Private/ChaosWheeledVehicleMovementComponent.cpp
@@ -662,9 +662,10 @@ void UChaosWheeledVehicleSimulation::ApplyInput(const FControlInputs& ControlInp
 			PWheel.SetBrakeTorque(MToCm(EngineBraking));
 		}
 
-		if (ModifiedInputs.HandbrakeInput && PWheel.Setup().HandbrakeEnabled)
+		if ((ModifiedInputs.HandbrakeInput && PWheel.Setup().HandbrakeEnabled) || ModifiedInputs.ParkingEnabled)
 		{
-			PWheel.SetBrakeTorque(MToCm(ModifiedInputs.HandbrakeInput * PWheel.Setup().HandbrakeTorque));
+			float HandbrakeForce = ModifiedInputs.ParkingEnabled ? PWheel.Setup().HandbrakeTorque : (ModifiedInputs.HandbrakeInput * PWheel.Setup().HandbrakeTorque);
+			PWheel.SetBrakeTorque(MToCm(HandbrakeForce));
 		}
 	}
 
diff --git a/Engine/Plugins/Experimental/ChaosVehiclesPlugin/Source/ChaosVehicles/Public/ChaosVehicleMovementComponent.h b/Engine/Plugins/Experimental/ChaosVehiclesPlugin/Source/ChaosVehicles/Public/ChaosVehicleMovementComponent.h
--- a/Engine/Plugins/Experimental/ChaosVehiclesPlugin/Source/ChaosVehicles/Public/ChaosVehicleMovementComponent.h
+++ b/Engine/Plugins/Experimental/ChaosVehiclesPlugin/Source/ChaosVehicles/Public/ChaosVehicleMovementComponent.h
@@ -37,6 +37,7 @@ struct FControlInputs
 		, RollInput(0.f)
 		, YawInput(0.f)
 		, HandbrakeInput(0.f)
+		, ParkingEnabled(false)
 		, TargetGearInput(0)
 		, SetGearImmediate(-1)
 		, TransmissionType(Chaos::ETransmissionType::Automatic)
@@ -67,6 +68,9 @@ struct FControlInputs
 	// Handbrake output to physics system. Range 0...1
 	float HandbrakeInput;
 
+	// ParkingEnabled
+	bool ParkingEnabled;
+
 	int32 TargetGearInput;
 	int8 SetGearImmediate;
 	Chaos::ETransmissionType TransmissionType;
@@ -769,6 +773,9 @@ protected:
 	UPROPERTY(Transient)
 	uint32 bWasAvoidanceUpdated : 1;
 
+	UPROPERTY(Transient)
+	uint8 bParkEnabled : 1;
+
 	int32 TargetGearInput;
 	int8 SetGearImmediate;
 	Chaos::ETransmissionType TransmissionType;
@@ -849,6 +856,10 @@ public:
 	UFUNCTION(BlueprintCallable, Category="Game|Components|ChaosVehicleMovement")
 	void SetHandbrakeInput(bool bNewHandbrake);
 
+	/** Set the vehicle in park mode */
+	UFUNCTION(BlueprintCallable, Category = "Game|Components|ChaosVehicleMovement")
+	void SetParked(bool bParked);
+
 	/** Set the user input for gear up */
 	UFUNCTION(BlueprintCallable, Category="Game|Components|ChaosVehicleMovement")
 	void SetChangeUpInput(bool bNewGearUp);
@@ -885,6 +896,10 @@ public:
 	UFUNCTION(BlueprintCallable, Category = "Game|Components|ChaosVehicleMovement")
 	float GetForwardSpeedMPH() const;
 
+	/** Is the vehicle in park mode */
+	UFUNCTION(BlueprintCallable, Category = "Game|Components|ChaosVehicleMovement")
+	bool IsParked() const;
+
 	UFUNCTION(BlueprintCallable, Category = "Game|Components|ChaosVehicleMovement")
 	void EnableSelfRighting(bool InState)
 	{
