ChaosVehicles: Simplify gear change request from Game thread after recent async changes

diff --git a/Engine/Plugins/Experimental/ChaosVehiclesPlugin/Source/ChaosVehicles/Private/ChaosVehicleMovementComponent.cpp b/Engine/Plugins/Experimental/ChaosVehiclesPlugin/Source/ChaosVehicles/Private/ChaosVehicleMovementComponent.cpp
--- a/Engine/Plugins/Experimental/ChaosVehiclesPlugin/Source/ChaosVehicles/Private/ChaosVehicleMovementComponent.cpp
+++ b/Engine/Plugins/Experimental/ChaosVehiclesPlugin/Source/ChaosVehicles/Private/ChaosVehicleMovementComponent.cpp
@@ -624,8 +624,6 @@ UChaosVehicleMovementComponent::UChaosVehicleMovementComponent(const FObjectInit
 	RollInputRate.FallRate = 10.0f;
 	YawInputRate.RiseRate = 6.0f;
 	YawInputRate.FallRate = 10.0f;
-	TargetGearInput = 0;
-	SetGearImmediate = -1;
 	TransmissionType = Chaos::ETransmissionType::Automatic;
 
 	SetIsReplicatedByDefault(true);
@@ -884,10 +882,20 @@ void UChaosVehicleMovementComponent::SetChangeDownInput(bool bNewGearDown)
 
 void UChaosVehicleMovementComponent::SetTargetGear(int32 GearNum, bool bImmediate)
 {
-	if (GearNum != TargetGearInput)
+	if (GearNum != PVehicleOutput->TargetGear)
 	{
-		TargetGearInput = GearNum;
-		SetGearImmediate = bImmediate?1:0; // true == 1, false == 0, don't transmit to physics thread == -1
+		FBodyInstance* TargetInstance = UpdatedPrimitive->GetBodyInstance();
+
+		if (TargetInstance)
+		{
+			FPhysicsCommand::ExecuteWrite(TargetInstance->ActorHandle, [&](const FPhysicsActorHandle& Chassis)
+				{
+					if (VehicleSimulationPT && VehicleSimulationPT->PVehicle && VehicleSimulationPT->PVehicle->HasTransmission())
+					{
+						VehicleSimulationPT->PVehicle->GetTransmission().SetGear(GearNum, bImmediate);
+					}
+				});
+		}
 	}
 }
 
@@ -1117,7 +1125,7 @@ void UChaosVehicleMovementComponent::UpdateState(float DeltaTime)
 		HandbrakeInput = HandbrakeInputRate.InterpInputValue(DeltaTime, HandbrakeInput, CalcHandbrakeInput());
 
 		// and send to server - (ServerUpdateState_Implementation below)
-		ServerUpdateState(SteeringInput, ThrottleInput, BrakeInput, HandbrakeInput, TargetGearInput, RollInput, PitchInput, YawInput);
+		ServerUpdateState(SteeringInput, ThrottleInput, BrakeInput, HandbrakeInput, GetTargetGear(), RollInput, PitchInput, YawInput);
 
 		if (PawnOwner && PawnOwner->IsNetMode(NM_Client))
 		{
@@ -1650,8 +1658,6 @@ void UChaosVehicleMovementComponent::Update(float DeltaTime)
 				AsyncInput->ControlInputs.YawInput = YawInput;
 				AsyncInput->ControlInputs.GearUpInput = bRawGearUpInput;
 				AsyncInput->ControlInputs.GearDownInput = bRawGearDownInput;
-				AsyncInput->ControlInputs.TargetGearInput = TargetGearInput;
-				AsyncInput->ControlInputs.SetGearImmediate = SetGearImmediate;
 				AsyncInput->ControlInputs.TransmissionType = TransmissionType;
 
 				AsyncInput->ControlInputs.ParkingEnabled = bParkEnabled;
@@ -1702,11 +1708,6 @@ void UChaosVehicleMovementComponent::ParallelUpdate(float DeltaSeconds)
 			PVehicleOutput->CurrentGear = CurAsyncOutput->VehicleSimOutput.CurrentGear;
 			PVehicleOutput->TargetGear = CurAsyncOutput->VehicleSimOutput.TargetGear;
 
-			if (TargetGearInput == PVehicleOutput->TargetGear)
-			{
-				SetGearImmediate = -1; // clear this otherwise can get gear change handshake issues between threads
-			}
-
 			// WHEN RUNNING WITH ASYNC ON & FIXED TIMESTEP THEN WE NEED TO INTERPOLATE BETWEEN THE CURRENT AND NEXT OUTPUT RESULTS
 			if (const FChaosVehicleAsyncOutput* NextOutput = static_cast<FChaosVehicleAsyncOutput*>(NextAsyncOutput))
 			{
diff --git a/Engine/Plugins/Experimental/ChaosVehiclesPlugin/Source/ChaosVehicles/Private/ChaosWheeledVehicleMovementComponent.cpp b/Engine/Plugins/Experimental/ChaosVehiclesPlugin/Source/ChaosVehicles/Private/ChaosWheeledVehicleMovementComponent.cpp
--- a/Engine/Plugins/Experimental/ChaosVehiclesPlugin/Source/ChaosVehicles/Private/ChaosWheeledVehicleMovementComponent.cpp
+++ b/Engine/Plugins/Experimental/ChaosVehiclesPlugin/Source/ChaosVehicles/Private/ChaosWheeledVehicleMovementComponent.cpp
@@ -611,11 +611,6 @@ void UChaosWheeledVehicleSimulation::ApplyInput(const FControlInputs& ControlInp
 			PTransmission.AccessSetup().TransmissionType = ModifiedInputs.TransmissionType;
 		}
 
-		if (ModifiedInputs.SetGearImmediate >= 0 && ModifiedInputs.TargetGearInput != PTransmission.GetTargetGear())
-		{
-			PTransmission.SetGear(ModifiedInputs.TargetGearInput, ModifiedInputs.SetGearImmediate?true:false);
-		}
-
 		if (ModifiedInputs.GearUpInput)
 		{
 			PTransmission.ChangeUp();
@@ -1459,8 +1454,8 @@ void UChaosWheeledVehicleMovementComponent::DrawDebug(UCanvas* Canvas, float& YL
 				, TransmissionSetup.ChangeUpRPM
 				, TransmissionSetup.ChangeDownRPM), 4, YPos);
 
-			YPos += Canvas->DrawText(RenderFont, FString::Printf(TEXT("Gear: %d (Target %d)(Target %d)")
-				, GetCurrentGear(), TargetGearInput, PVehicleOutput->TargetGear), 4, YPos);
+			YPos += Canvas->DrawText(RenderFont, FString::Printf(TEXT("Gear: %d (Target %d)")
+				, GetCurrentGear(), GetTargetGear()), 4, YPos);
 		}
 		//YPos += Canvas->DrawText(RenderFont, FString::Printf(TEXT("Drag: %.1f"), DebugDragMagnitude), 4, YPos);
 
diff --git a/Engine/Plugins/Experimental/ChaosVehiclesPlugin/Source/ChaosVehicles/Public/ChaosVehicleMovementComponent.h b/Engine/Plugins/Experimental/ChaosVehiclesPlugin/Source/ChaosVehicles/Public/ChaosVehicleMovementComponent.h
--- a/Engine/Plugins/Experimental/ChaosVehiclesPlugin/Source/ChaosVehicles/Public/ChaosVehicleMovementComponent.h
+++ b/Engine/Plugins/Experimental/ChaosVehiclesPlugin/Source/ChaosVehicles/Public/ChaosVehicleMovementComponent.h
@@ -38,8 +38,6 @@ struct FControlInputs
 		, YawInput(0.f)
 		, HandbrakeInput(0.f)
 		, ParkingEnabled(false)
-		, TargetGearInput(0)
-		, SetGearImmediate(-1)
 		, TransmissionType(Chaos::ETransmissionType::Automatic)
 		, GearUpInput(false)
 		, GearDownInput(false)
@@ -71,11 +69,8 @@ struct FControlInputs
 	// ParkingEnabled
 	bool ParkingEnabled;
 
-	int32 TargetGearInput;
-	int8 SetGearImmediate;
 	Chaos::ETransmissionType TransmissionType;
 
-
 	bool GearUpInput;
 	bool GearDownInput;
 };
@@ -776,8 +771,6 @@ protected:
 	UPROPERTY(Transient)
 	uint8 bParkEnabled : 1;
 
-	int32 TargetGearInput;
-	int8 SetGearImmediate;
 	Chaos::ETransmissionType TransmissionType;
 
 public:
diff --git a/Engine/Source/Programs/HeadlessChaos/Private/HeadlessChaosTestVehicles.cpp b/Engine/Source/Programs/HeadlessChaos/Private/HeadlessChaosTestVehicles.cpp
--- a/Engine/Source/Programs/HeadlessChaos/Private/HeadlessChaosTestVehicles.cpp
+++ b/Engine/Source/Programs/HeadlessChaos/Private/HeadlessChaosTestVehicles.cpp
@@ -418,6 +418,7 @@ namespace ChaosTest
 			Setup.ForwardRatios.Add(2.f);
 			Setup.ForwardRatios.Add(1.f);
 			Setup.ReverseRatios.Add(3.f);
+			Setup.ReverseRatios.Add(6.f);
 			Setup.FinalDriveRatio = 4.f;
 			Setup.ChangeUpRPM = 3000;
 			Setup.ChangeDownRPM = 1200;
@@ -453,8 +454,17 @@ namespace ChaosTest
 		EXPECT_EQ(Transmission.GetCurrentGear(), -1);
 
 		Transmission.ChangeDown();
+		EXPECT_EQ(Transmission.GetCurrentGear(), -2);
+
+		Transmission.ChangeDown();
+		EXPECT_EQ(Transmission.GetCurrentGear(), -2);
+
+		Transmission.ChangeUp();
 		EXPECT_EQ(Transmission.GetCurrentGear(), -1);
 
+		Transmission.ChangeUp();
+		EXPECT_EQ(Transmission.GetCurrentGear(), 0);
+
 		Transmission.SetGear(1);
 
 		// Now change settings so we have a delay in the gear changing
@@ -488,6 +498,7 @@ namespace ChaosTest
 			Setup.ForwardRatios.Add(2.f);
 			Setup.ForwardRatios.Add(1.f);
 			Setup.ReverseRatios.Add(3.f);
+			Setup.ReverseRatios.Add(6.f);
 			Setup.FinalDriveRatio = 4.f;
 			Setup.ChangeUpRPM = 3000;
 			Setup.ChangeDownRPM = 1200;
@@ -520,6 +531,31 @@ namespace ChaosTest
 		Transmission.Simulate(0.25f);
 		EXPECT_EQ(Transmission.GetCurrentGear(), 1);
 
+		// stays in first, doesn't change to neutral
+		Transmission.SetEngineRPM(1000);
+		Transmission.Simulate(0.25f);
+		EXPECT_EQ(Transmission.GetCurrentGear(), 1);
+
+		Transmission.SetGear(-2, true);
+
+		Transmission.SetEngineRPM(3000);
+		Transmission.Simulate(0.25f);
+		EXPECT_EQ(Transmission.GetCurrentGear(), -2);
+
+		Transmission.SetEngineRPM(1000);
+		Transmission.Simulate(0.25f);
+		EXPECT_EQ(Transmission.GetCurrentGear(), -1);
+
+		// stays in reverse first, doesn't change to neutral
+		Transmission.SetEngineRPM(1000);
+		Transmission.Simulate(0.25f);
+		EXPECT_EQ(Transmission.GetCurrentGear(), -1);
+
+		// changes to next reverse gear
+		Transmission.SetEngineRPM(3000);
+		Transmission.Simulate(0.25f);
+		EXPECT_EQ(Transmission.GetCurrentGear(), -2);
+
 	}
 
 	TYPED_TEST(AllTraits, VehicleTest_TransmissionGearRatios)
diff --git a/Engine/Source/Runtime/Experimental/ChaosVehicles/ChaosVehiclesCore/Private/TransmissionSystem.cpp b/Engine/Source/Runtime/Experimental/ChaosVehicles/ChaosVehiclesCore/Private/TransmissionSystem.cpp
--- a/Engine/Source/Runtime/Experimental/ChaosVehicles/ChaosVehiclesCore/Private/TransmissionSystem.cpp
+++ b/Engine/Source/Runtime/Experimental/ChaosVehicles/ChaosVehiclesCore/Private/TransmissionSystem.cpp
@@ -24,14 +24,17 @@ namespace Chaos
 
 		TargetGear = InGear;
 
-		if (Immediate || Setup().GearChangeTime == 0.f)
+		if (TargetGear != CurrentGear)
 		{
-			CurrentGear = TargetGear;
-		}
-		else
-		{
-			CurrentGear = 0;	// go through neutral for GearChangeTime time period
-			CurrentGearChangeTime = Setup().GearChangeTime;
+			if (Immediate || Setup().GearChangeTime == 0.f)
+			{
+				CurrentGear = TargetGear;
+			}
+			else
+			{
+				CurrentGear = 0;	// go through neutral for GearChangeTime time period
+				CurrentGearChangeTime = Setup().GearChangeTime;
+			}
 		}
 	}
 
@@ -62,11 +65,25 @@ namespace Chaos
 			{
 				if (EngineRPM >= Setup().ChangeUpRPM)
 				{
-					ChangeUp();
+					if (CurrentGear > 0)
+					{
+						ChangeUp();
+					}
+					else
+					{
+						ChangeDown();
+					}
 				}
-				else if (EngineRPM <= Setup().ChangeDownRPM && CurrentGear > 1) // don't change down to neutral
+				else if (EngineRPM <= Setup().ChangeDownRPM && FMath::Abs(CurrentGear) > 1) // don't change down to neutral
 				{
-					ChangeDown();
+					if (CurrentGear > 0)
+					{
+						ChangeDown();
+					}
+					else
+					{
+						ChangeUp();
+					}
 				}
 			}
 		}
