ChaosVehicles: Wheel Class swappable at runtime via blueprint. Call set SetWheelClass on the ChaosWheeledVehicleComponent specifying the wheels index and the new wheel blueprint class.

diff --git a/Engine/Plugins/Experimental/ChaosVehiclesPlugin/Source/ChaosVehicles/Private/ChaosVehicleMovementComponent.cpp b/Engine/Plugins/Experimental/ChaosVehiclesPlugin/Source/ChaosVehicles/Private/ChaosVehicleMovementComponent.cpp
--- a/Engine/Plugins/Experimental/ChaosVehiclesPlugin/Source/ChaosVehicles/Private/ChaosVehicleMovementComponent.cpp
+++ b/Engine/Plugins/Experimental/ChaosVehiclesPlugin/Source/ChaosVehicles/Private/ChaosVehicleMovementComponent.cpp
@@ -583,6 +583,15 @@ void UChaosVehicleSimulation::AddTorqueInRadians(const FVector& Torque, bool bAl
 	}
 }
 
+void UChaosVehicleSimulation::InitializeWheel(int WheelIndex, const Chaos::FSimpleWheelConfig* InWheelSetup)
+{
+	if (PVehicle->IsValid() && InWheelSetup && WheelIndex < PVehicle->Wheels.Num())
+	{
+		PVehicle->Wheels[WheelIndex].SetupPtr = InWheelSetup;
+		PVehicle->Wheels[WheelIndex].SetWheelRadius(InWheelSetup->WheelRadius);
+	}
+
+}
 
 /**
  * UChaosVehicleMovementComponent
diff --git a/Engine/Plugins/Experimental/ChaosVehiclesPlugin/Source/ChaosVehicles/Private/ChaosWheeledVehicleMovementComponent.cpp b/Engine/Plugins/Experimental/ChaosVehiclesPlugin/Source/ChaosVehicles/Private/ChaosWheeledVehicleMovementComponent.cpp
--- a/Engine/Plugins/Experimental/ChaosVehiclesPlugin/Source/ChaosVehicles/Private/ChaosWheeledVehicleMovementComponent.cpp
+++ b/Engine/Plugins/Experimental/ChaosVehiclesPlugin/Source/ChaosVehicles/Private/ChaosWheeledVehicleMovementComponent.cpp
@@ -1917,6 +1917,30 @@ void UChaosWheeledVehicleMovementComponent::ParallelUpdate(float DeltaSeconds)
 	FillWheelOutputState(); // expose wheel/suspensionm data to blueprint
 }
 
+void UChaosWheeledVehicleMovementComponent::SetWheelClass(int WheelIndex, TSubclassOf<UChaosVehicleWheel> InWheelClass)
+{
+	if (UpdatedPrimitive && InWheelClass)
+	{
+		FBodyInstance* TargetInstance = UpdatedPrimitive->GetBodyInstance();
+
+		if (TargetInstance && WheelIndex < Wheels.Num())
+		{
+			UChaosVehicleWheel* Wheel = Wheels[WheelIndex];
+			Wheel->Shutdown();
+			Wheel = NewObject<UChaosVehicleWheel>(this, InWheelClass);
+			Wheel->Init(this, WheelIndex);
+
+			FPhysicsCommand::ExecuteWrite(TargetInstance->ActorHandle, [&](const FPhysicsActorHandle& Chassis)
+				{
+					if (VehicleSimulationPT)
+					{
+						VehicleSimulationPT->InitializeWheel(WheelIndex, &Wheel->GetPhysicsWheelConfig());
+					}
+				});
+		}
+	}
+}
+
 #endif
 
 FChaosWheelSetup::FChaosWheelSetup()
diff --git a/Engine/Plugins/Experimental/ChaosVehiclesPlugin/Source/ChaosVehicles/Public/ChaosVehicleMovementComponent.h b/Engine/Plugins/Experimental/ChaosVehiclesPlugin/Source/ChaosVehicles/Public/ChaosVehicleMovementComponent.h
--- a/Engine/Plugins/Experimental/ChaosVehiclesPlugin/Source/ChaosVehicles/Public/ChaosVehicleMovementComponent.h
+++ b/Engine/Plugins/Experimental/ChaosVehiclesPlugin/Source/ChaosVehicles/Public/ChaosVehicleMovementComponent.h
@@ -654,6 +654,8 @@ public:
 	/** Add a torque to this vehicle */
 	void AddTorqueInRadians(const FVector& Torque, bool bAllowSubstepping = true, bool bAccelChange = false);
 
+	/** Reinitialize a wheel at runtime */
+	void InitializeWheel(int WheelIndex, const Chaos::FSimpleWheelConfig* InWheelSetup);
 
 	/** Draw debug text for the wheels and suspension */
 	virtual void DrawDebug3D();
diff --git a/Engine/Plugins/Experimental/ChaosVehiclesPlugin/Source/ChaosVehicles/Public/ChaosWheeledVehicleMovementComponent.h b/Engine/Plugins/Experimental/ChaosVehiclesPlugin/Source/ChaosVehicles/Public/ChaosWheeledVehicleMovementComponent.h
--- a/Engine/Plugins/Experimental/ChaosVehiclesPlugin/Source/ChaosVehicles/Public/ChaosWheeledVehicleMovementComponent.h
+++ b/Engine/Plugins/Experimental/ChaosVehiclesPlugin/Source/ChaosVehicles/Public/ChaosWheeledVehicleMovementComponent.h
@@ -707,6 +707,9 @@ class CHAOSVEHICLES_API UChaosWheeledVehicleMovementComponent : public UChaosVeh
 		bWheelFrictionEnabled = InState;
 	}
 
+	UFUNCTION(BlueprintCallable, Category = "Game|Components|ChaosWheeledVehicleMovement")
+	void SetWheelClass(int WheelIndex, TSubclassOf<UChaosVehicleWheel> InWheelClass);
+
 
 protected:
 
