CIS Static analysis fix

diff --git a/Engine/Plugins/Experimental/ChaosVehiclesPlugin/Source/ChaosVehicles/Private/ChaosVehicleMovementComponent.cpp b/Engine/Plugins/Experimental/ChaosVehiclesPlugin/Source/ChaosVehicles/Private/ChaosVehicleMovementComponent.cpp
--- a/Engine/Plugins/Experimental/ChaosVehiclesPlugin/Source/ChaosVehicles/Private/ChaosVehicleMovementComponent.cpp
+++ b/Engine/Plugins/Experimental/ChaosVehiclesPlugin/Source/ChaosVehicles/Private/ChaosVehicleMovementComponent.cpp
@@ -247,9 +247,16 @@ void UChaosVehicleSimulation::ApplyAerodynamics(float DeltaTime)
 
 FVector GetWorldVelocityAtPoint(const Chaos::FRigidBodyHandle_Internal* RigidHandle, const FVector& WorldLocation)
 {
-	const Chaos::FVec3 COM = RigidHandle ? Chaos::FParticleUtilitiesGT::GetCoMWorldPosition(RigidHandle) : Chaos::FParticleUtilitiesGT::GetActorWorldTransform(RigidHandle).GetTranslation();
-	const Chaos::FVec3 Diff = WorldLocation - COM;
-	return RigidHandle->V() - Chaos::FVec3::CrossProduct(Diff, RigidHandle->W());
+	if (RigidHandle)
+	{
+		const Chaos::FVec3 COM = RigidHandle ? Chaos::FParticleUtilitiesGT::GetCoMWorldPosition(RigidHandle) : Chaos::FParticleUtilitiesGT::GetActorWorldTransform(RigidHandle).GetTranslation();
+		const Chaos::FVec3 Diff = WorldLocation - COM;
+		return RigidHandle->V() - Chaos::FVec3::CrossProduct(Diff, RigidHandle->W());
+	}
+	else
+	{
+		return FVector::ZeroVector;
+	}
 }
 
 void UChaosVehicleSimulation::ApplyAerofoilForces(float DeltaTime)
diff --git a/Engine/Plugins/Experimental/ChaosVehiclesPlugin/Source/ChaosVehicles/Private/ChaosWheeledVehicleMovementComponent.cpp b/Engine/Plugins/Experimental/ChaosVehiclesPlugin/Source/ChaosVehicles/Private/ChaosWheeledVehicleMovementComponent.cpp
--- a/Engine/Plugins/Experimental/ChaosVehiclesPlugin/Source/ChaosVehicles/Private/ChaosWheeledVehicleMovementComponent.cpp
+++ b/Engine/Plugins/Experimental/ChaosVehiclesPlugin/Source/ChaosVehicles/Private/ChaosWheeledVehicleMovementComponent.cpp
@@ -99,9 +99,16 @@ void FWheelState::CaptureState(int WheelIdx, const FVector& WheelOffset, const C
 
 FVector FWheelState::GetVelocityAtPoint(const Chaos::FRigidBodyHandle_Internal* Rigid, const FVector& InPoint)
 {
-	const Chaos::FVec3 COM = Rigid ? Chaos::FParticleUtilitiesGT::GetCoMWorldPosition(Rigid) : Chaos::FParticleUtilitiesGT::GetActorWorldTransform(Rigid).GetTranslation();
-	const Chaos::FVec3 Diff = InPoint - COM;
-	return Rigid->V() - Chaos::FVec3::CrossProduct(Diff, Rigid->W());
+	if (Rigid)
+	{
+		const Chaos::FVec3 COM = Rigid ? Chaos::FParticleUtilitiesGT::GetCoMWorldPosition(Rigid) : Chaos::FParticleUtilitiesGT::GetActorWorldTransform(Rigid).GetTranslation();
+		const Chaos::FVec3 Diff = InPoint - COM;
+		return Rigid->V() - Chaos::FVec3::CrossProduct(Diff, Rigid->W());
+	}
+	else
+	{
+		return FVector::ZeroVector;
+	}
 }
 
 /**
