ChaosVehicles: Vehicle state Get and Set Snapshot of dynamic data.

diff --git a/Engine/Plugins/Experimental/ChaosVehiclesPlugin/Source/ChaosVehicles/Private/ChaosVehicleMovementComponent.cpp b/Engine/Plugins/Experimental/ChaosVehiclesPlugin/Source/ChaosVehicles/Private/ChaosVehicleMovementComponent.cpp
--- a/Engine/Plugins/Experimental/ChaosVehiclesPlugin/Source/ChaosVehicles/Private/ChaosVehicleMovementComponent.cpp
+++ b/Engine/Plugins/Experimental/ChaosVehiclesPlugin/Source/ChaosVehicles/Private/ChaosVehicleMovementComponent.cpp
@@ -1280,6 +1280,10 @@ FBodyInstance* UChaosVehicleMovementComponent::GetBodyInstance()
 	return UpdatedPrimitive ? UpdatedPrimitive->GetBodyInstance() : nullptr;
 }
 
+const FBodyInstance* UChaosVehicleMovementComponent::GetBodyInstance() const
+{
+	return UpdatedPrimitive ? UpdatedPrimitive->GetBodyInstance() : nullptr;
+}
 
 UMeshComponent* UChaosVehicleMovementComponent::GetMesh() const
 {
@@ -1546,9 +1550,6 @@ void UChaosVehicleMovementComponent::SetCurrentAsyncInputOutputInternal(FChaosVe
 	NextAsyncOutput = nullptr;
 	OutputInterpAlpha = 0.f;
 
-	//FChaosVehicleDefaultAsyncInput* CurDefaultAsyncInput = static_cast<FChaosVehicleDefaultAsyncInput*>(CurAsyncInput);
-	//CurDefaultAsyncInput->ControlInputs.SteeringInput = 1.0f;
-
 	// We need to find our vehicle in the output given
 	if (CurOutput)
 	{
@@ -1739,8 +1740,10 @@ void UChaosVehicleMovementComponent::ParallelUpdate(float DeltaSeconds)
 
 					PVehicleOutput->Wheels[WheelIdx].InContact = Current.InContact;
 					PVehicleOutput->Wheels[WheelIdx].SteeringAngle = FMath::Lerp(Current.SteeringAngle, Next.SteeringAngle, OutputInterpAlpha);
+					PVehicleOutput->Wheels[WheelIdx].WheelRadius = FMath::Lerp(Current.WheelRadius, Next.WheelRadius, OutputInterpAlpha);
 					float DeltaAngle = FMath::FindDeltaAngleRadians(Current.AngularPosition, Next.AngularPosition);
 					PVehicleOutput->Wheels[WheelIdx].AngularPosition = Current.AngularPosition + DeltaAngle * OutputInterpAlpha;
+					PVehicleOutput->Wheels[WheelIdx].AngularVelocity = FMath::Lerp(Current.AngularVelocity, Next.AngularVelocity, OutputInterpAlpha);
 					PVehicleOutput->Wheels[WheelIdx].LateralAdhesiveLimit = FMath::Lerp(Current.LateralAdhesiveLimit, Next.LateralAdhesiveLimit, OutputInterpAlpha);
 					PVehicleOutput->Wheels[WheelIdx].LongitudinalAdhesiveLimit = FMath::Lerp(Current.LongitudinalAdhesiveLimit, Next.LongitudinalAdhesiveLimit, OutputInterpAlpha);
 
@@ -1768,7 +1771,9 @@ void UChaosVehicleMovementComponent::ParallelUpdate(float DeltaSeconds)
 
 					PVehicleOutput->Wheels[WheelIdx].InContact = Current.InContact;
 					PVehicleOutput->Wheels[WheelIdx].SteeringAngle = Current.SteeringAngle;
+					PVehicleOutput->Wheels[WheelIdx].WheelRadius = Current.WheelRadius;
 					PVehicleOutput->Wheels[WheelIdx].AngularPosition = Current.AngularPosition;
+					PVehicleOutput->Wheels[WheelIdx].AngularVelocity = Current.AngularVelocity;
 					PVehicleOutput->Wheels[WheelIdx].LateralAdhesiveLimit = Current.LateralAdhesiveLimit;
 					PVehicleOutput->Wheels[WheelIdx].LongitudinalAdhesiveLimit = Current.LongitudinalAdhesiveLimit;
 
@@ -1789,6 +1794,26 @@ void UChaosVehicleMovementComponent::ParallelUpdate(float DeltaSeconds)
 	}
 }
 
+void UChaosVehicleMovementComponent::GetBaseSnapshot(FBaseSnapshotData& SnapshotOut) const
+{
+	if (const FBodyInstance* TargetInstance = GetBodyInstance())
+	{
+		SnapshotOut.Transform = TargetInstance->GetUnrealWorldTransform();
+		SnapshotOut.LinearVelocity = TargetInstance->GetUnrealWorldVelocity();
+		SnapshotOut.AngularVelocity = TargetInstance->GetUnrealWorldAngularVelocityInRadians();
+	}
+}
+
+void UChaosVehicleMovementComponent::SetBaseSnapshot(const FBaseSnapshotData& SnapshotIn)
+{
+	if (FBodyInstance* TargetInstance = GetBodyInstance())
+	{
+		TargetInstance->SetLinearVelocity(SnapshotIn.LinearVelocity, false);
+		TargetInstance->SetAngularVelocityInRadians(SnapshotIn.AngularVelocity, false);
+		TargetInstance->SetBodyTransform(SnapshotIn.Transform, ETeleportType::TeleportPhysics);
+	}
+}
+
 #undef LOCTEXT_NAMESPACE
 
 
diff --git a/Engine/Plugins/Experimental/ChaosVehiclesPlugin/Source/ChaosVehicles/Private/ChaosVehicleWheel.cpp b/Engine/Plugins/Experimental/ChaosVehiclesPlugin/Source/ChaosVehicles/Private/ChaosVehicleWheel.cpp
--- a/Engine/Plugins/Experimental/ChaosVehiclesPlugin/Source/ChaosVehicles/Private/ChaosVehicleWheel.cpp
+++ b/Engine/Plugins/Experimental/ChaosVehiclesPlugin/Source/ChaosVehicles/Private/ChaosVehicleWheel.cpp
@@ -69,14 +69,21 @@ float UChaosVehicleWheel::GetRotationAngle() const
 	check(VehicleComponent && VehicleComponent->PhysicsVehicleOutput());
 	float RotationAngle = -1.0f * FMath::RadiansToDegrees(VehicleComponent->PhysicsVehicleOutput()->Wheels[WheelIndex].AngularPosition);
 	ensure(!FMath::IsNaN(RotationAngle));
-
-	//if (WheelIndex == 2)
-	//{
-	//	UE_LOG(LogChaos, Warning, TEXT("RotationAngle = %f"), RotationAngle);
-	//}
 	return RotationAngle;
 }
 
+float UChaosVehicleWheel::GetWheelRadius() const
+{
+	check(VehicleComponent && VehicleComponent->PhysicsVehicleOutput());
+	return VehicleComponent->PhysicsVehicleOutput()->Wheels[WheelIndex].WheelRadius;
+}
+
+float UChaosVehicleWheel::GetWheelAngularVelocity() const
+{
+	check(VehicleComponent && VehicleComponent->PhysicsVehicleOutput());
+	return VehicleComponent->PhysicsVehicleOutput()->Wheels[WheelIndex].AngularVelocity;
+}
+
 float UChaosVehicleWheel::GetSuspensionOffset() const
 {
 	check(VehicleComponent && VehicleComponent->PhysicsVehicleOutput());
diff --git a/Engine/Plugins/Experimental/ChaosVehiclesPlugin/Source/ChaosVehicles/Private/ChaosWheeledVehicleMovementComponent.cpp b/Engine/Plugins/Experimental/ChaosVehiclesPlugin/Source/ChaosVehicles/Private/ChaosWheeledVehicleMovementComponent.cpp
--- a/Engine/Plugins/Experimental/ChaosVehiclesPlugin/Source/ChaosVehicles/Private/ChaosWheeledVehicleMovementComponent.cpp
+++ b/Engine/Plugins/Experimental/ChaosVehiclesPlugin/Source/ChaosVehicles/Private/ChaosWheeledVehicleMovementComponent.cpp
@@ -849,6 +849,8 @@ void UChaosWheeledVehicleSimulation::FillOutputState(FChaosVehicleAsyncOutput& O
 		WheelsOut.InContact = VehicleWheels[WheelIdx].InContact();
 		WheelsOut.SteeringAngle = VehicleWheels[WheelIdx].GetSteeringAngle();
 		WheelsOut.AngularPosition = VehicleWheels[WheelIdx].GetAngularPosition();
+		WheelsOut.AngularVelocity = VehicleWheels[WheelIdx].GetAngularVelocity();
+		WheelsOut.WheelRadius = VehicleWheels[WheelIdx].GetEffectiveRadius();
 
 		WheelsOut.LateralAdhesiveLimit = VehicleWheels[WheelIdx].LateralAdhesiveLimit;
 		WheelsOut.LongitudinalAdhesiveLimit = VehicleWheels[WheelIdx].LongitudinalAdhesiveLimit;
@@ -1918,18 +1920,68 @@ FWheelStatus UChaosWheeledVehicleMovementComponent::MakeWheelStatus(bool bInCont
 	return Status;
 }
 
+void UChaosWheeledVehicleMovementComponent::BreakWheeledSnapshot(const struct FWheeledSnaphotData& Snapshot, FTransform& Transform, FVector& LinearVelocity
+	, FVector& AngularVelocity, int& SelectedGear, float& EngineRPM, TArray<FWheelSnapshot>& WheelSnapshots)
+{
+	Transform = Snapshot.Transform;
+	LinearVelocity = Snapshot.LinearVelocity;
+	AngularVelocity = Snapshot.AngularVelocity;
+	SelectedGear = Snapshot.SelectedGear;
+	EngineRPM = Snapshot.EngineRPM;
+	WheelSnapshots = Snapshot.WheelSnapshots;
+}
+
+FWheeledSnaphotData UChaosWheeledVehicleMovementComponent::MakeWheeledSnapshot(FTransform Transform, FVector LinearVelocity
+	, FVector AngularVelocity, int SelectedGear, float EngineRPM, TArray<FWheelSnapshot>& WheelSnapshots)
+{
+	FWheeledSnaphotData Snapshot;
+	Snapshot.Transform = Transform;
+	Snapshot.LinearVelocity = LinearVelocity;
+	Snapshot.AngularVelocity = AngularVelocity;
+	Snapshot.SelectedGear = SelectedGear;
+	Snapshot.EngineRPM = EngineRPM;
+	Snapshot.WheelSnapshots = WheelSnapshots;
+
+	return Snapshot;
+}
+
+
+void UChaosWheeledVehicleMovementComponent::BreakWheelSnapshot(const struct FWheelSnapshot& Snapshot, float& SuspensionOffset
+	, float& WheelRotationAngle, float& SteeringAngle, float& WheelRadius, float& WheelAngularVelocity)
+{
+	SuspensionOffset = Snapshot.SuspensionOffset;
+	WheelRotationAngle = Snapshot.WheelRotationAngle;
+	SteeringAngle = Snapshot.SteeringAngle;
+	WheelRadius = Snapshot.WheelRadius;
+	WheelAngularVelocity = Snapshot.WheelAngularVelocity;
+}
+
+FWheelSnapshot UChaosWheeledVehicleMovementComponent::MakeWheelSnapshot(float SuspensionOffset, float WheelRotationAngle
+	, float SteeringAngle, float WheelRadius, float WheelAngularVelocity)
+{
+	FWheelSnapshot Snapshot;
+	Snapshot.SuspensionOffset = SuspensionOffset;
+	Snapshot.WheelRotationAngle = WheelRotationAngle;
+	Snapshot.SteeringAngle = SteeringAngle;
+	Snapshot.WheelRadius = WheelRadius;
+	Snapshot.WheelAngularVelocity = WheelAngularVelocity;
+
+	return Snapshot;
+}
+
+
 void UChaosWheeledVehicleMovementComponent::ParallelUpdate(float DeltaSeconds)
 {
 	UChaosVehicleMovementComponent::ParallelUpdate(DeltaSeconds);\
 
-	FillWheelOutputState(); // expose wheel/suspensionm data to blueprint
+	FillWheelOutputState(); // exposes wheel/suspension data to blueprint
 }
 
 void UChaosWheeledVehicleMovementComponent::SetWheelClass(int WheelIndex, TSubclassOf<UChaosVehicleWheel> InWheelClass)
 {
 	if (UpdatedPrimitive && InWheelClass)
 	{
-		FBodyInstance* TargetInstance = UpdatedPrimitive->GetBodyInstance();
+		FBodyInstance* TargetInstance = GetBodyInstance();
 
 		if (TargetInstance && WheelIndex < Wheels.Num())
 		{
@@ -1949,6 +2001,65 @@ void UChaosWheeledVehicleMovementComponent::SetWheelClass(int WheelIndex, TSubcl
 	}
 }
 
+FWheeledSnaphotData UChaosWheeledVehicleMovementComponent::GetSnapshot() const
+{
+	FWheeledSnaphotData WheelSnapshotData;
+	UChaosVehicleMovementComponent::GetBaseSnapshot(WheelSnapshotData);
+
+	WheelSnapshotData.EngineRPM = PVehicleOutput->EngineRPM;
+	WheelSnapshotData.SelectedGear = PVehicleOutput->CurrentGear;
+	WheelSnapshotData.WheelSnapshots.SetNum(Wheels.Num());
+
+	int WheelIdx = 0;
+	for (const UChaosVehicleWheel* Wheel : Wheels)
+	{
+		WheelSnapshotData.WheelSnapshots[WheelIdx].SteeringAngle = Wheel->GetSteerAngle();
+		WheelSnapshotData.WheelSnapshots[WheelIdx].SuspensionOffset = Wheel->GetSuspensionOffset();
+		WheelSnapshotData.WheelSnapshots[WheelIdx].WheelRotationAngle = Wheel->GetRotationAngle();
+		WheelSnapshotData.WheelSnapshots[WheelIdx].WheelRadius = Wheel->GetWheelRadius();
+		WheelSnapshotData.WheelSnapshots[WheelIdx].WheelAngularVelocity = Wheel->GetWheelAngularVelocity();
+
+		WheelIdx++;
+	}
+
+	return WheelSnapshotData;
+}
+
+void UChaosWheeledVehicleMovementComponent::SetSnapshot(const FWheeledSnaphotData& SnapshotIn)
+{
+	UChaosVehicleMovementComponent::SetBaseSnapshot(SnapshotIn);
+
+	FBodyInstance* TargetInstance = GetBodyInstance();
+
+	FPhysicsCommand::ExecuteWrite(TargetInstance->ActorHandle, [&](const FPhysicsActorHandle& Chassis)
+		{
+			const FWheeledSnaphotData* WheelSnapshotData = static_cast<const FWheeledSnaphotData*>(&SnapshotIn);
+
+			ensure(WheelSnapshotData->WheelSnapshots.Num() == VehicleSimulationPT->PVehicle->Wheels.Num());
+			if (VehicleSimulationPT && VehicleSimulationPT->PVehicle)
+			{
+				for (int WheelIdx = 0; WheelIdx < WheelSnapshotData->WheelSnapshots.Num(); WheelIdx++)
+				{
+					ensure(VehicleSimulationPT->PVehicle->Wheels.Num() == VehicleSimulationPT->PVehicle->Suspension.Num());
+
+					if (WheelIdx < VehicleSimulationPT->PVehicle->Wheels.Num())
+					{
+						const FWheelSnapshot& Data = WheelSnapshotData->WheelSnapshots[WheelIdx];
+						Chaos::FSimpleWheelSim& VehicleWheel = VehicleSimulationPT->PVehicle->Wheels[WheelIdx];
+						Chaos::FSimpleSuspensionSim& VehicleSuspension = VehicleSimulationPT->PVehicle->Suspension[WheelIdx];
+
+						VehicleWheel.SetSteeringAngle(Data.SteeringAngle);
+						VehicleSuspension.SetSuspensionLength(Data.SuspensionOffset, Data.WheelRadius);
+						VehicleWheel.SetAngularPosition(FMath::DegreesToRadians(-Data.WheelRotationAngle));
+						VehicleWheel.SetWheelRadius(Data.WheelRadius);
+						VehicleWheel.SetAngularVelocity(Data.WheelAngularVelocity);
+					}
+				}
+			}
+		});
+
+}
+
 #endif
 
 FChaosWheelSetup::FChaosWheelSetup()
diff --git a/Engine/Plugins/Experimental/ChaosVehiclesPlugin/Source/ChaosVehicles/Public/AsyncCallback.h b/Engine/Plugins/Experimental/ChaosVehiclesPlugin/Source/ChaosVehicles/Public/AsyncCallback.h
--- a/Engine/Plugins/Experimental/ChaosVehiclesPlugin/Source/ChaosVehicles/Public/AsyncCallback.h
+++ b/Engine/Plugins/Experimental/ChaosVehiclesPlugin/Source/ChaosVehicles/Public/AsyncCallback.h
@@ -23,6 +23,8 @@ struct FWheelsOutput
 		: InContact(false)
 		, SteeringAngle(0.f)
 		, AngularPosition(0.f)
+		, AngularVelocity(0.f)
+		, WheelRadius(0.f)
 		, LateralAdhesiveLimit(0.f)
 		, LongitudinalAdhesiveLimit(0.f)
 		, bIsSlipping(false)
@@ -40,6 +42,8 @@ struct FWheelsOutput
 	bool InContact;
 	float SteeringAngle;
 	float AngularPosition;
+	float AngularVelocity;
+	float WheelRadius;
 
 	float LateralAdhesiveLimit;
 	float LongitudinalAdhesiveLimit;
diff --git a/Engine/Plugins/Experimental/ChaosVehiclesPlugin/Source/ChaosVehicles/Public/ChaosVehicleMovementComponent.h b/Engine/Plugins/Experimental/ChaosVehiclesPlugin/Source/ChaosVehicles/Public/ChaosVehicleMovementComponent.h
--- a/Engine/Plugins/Experimental/ChaosVehiclesPlugin/Source/ChaosVehicles/Public/ChaosVehicleMovementComponent.h
+++ b/Engine/Plugins/Experimental/ChaosVehiclesPlugin/Source/ChaosVehicles/Public/ChaosVehicleMovementComponent.h
@@ -13,7 +13,7 @@
 #include "AerofoilSystem.h"
 #include "ThrustSystem.h"
 #include "PhysicsProxy/SingleParticlePhysicsProxyFwd.h"
-
+#include "SnapshotData.h"
 #include "AsyncCallback.h"
 
 #include "ChaosVehicleMovementComponent.generated.h"
@@ -896,6 +896,12 @@ public:
 	UFUNCTION(BlueprintCallable, Category = "Game|Components|ChaosVehicleMovement")
 	bool IsParked() const;
 
+	/** Grab a snapshot of the vehicle instance dynamic state */
+	void GetBaseSnapshot(FBaseSnapshotData& SnapshotOut) const;
+
+	/** Set snapshot of vehicle instance dynamic state */
+	void SetBaseSnapshot(const FBaseSnapshotData& SnapshotIn);
+
 	UFUNCTION(BlueprintCallable, Category = "Game|Components|ChaosVehicleMovement")
 	void EnableSelfRighting(bool InState)
 	{
@@ -1155,7 +1161,7 @@ protected:
 	}
 
 	FBodyInstance* GetBodyInstance();
-
+	const FBodyInstance* GetBodyInstance() const;
 	FVehicleState VehicleState;							/* Useful vehicle state captured at start of frame */
 	TUniquePtr<FPhysicsVehicleOutput> PVehicleOutput;	/* physics simulation data output from the async physics thread */
 
diff --git a/Engine/Plugins/Experimental/ChaosVehiclesPlugin/Source/ChaosVehicles/Public/ChaosVehicleWheel.h b/Engine/Plugins/Experimental/ChaosVehiclesPlugin/Source/ChaosVehicles/Public/ChaosVehicleWheel.h
--- a/Engine/Plugins/Experimental/ChaosVehiclesPlugin/Source/ChaosVehicles/Public/ChaosVehicleWheel.h
+++ b/Engine/Plugins/Experimental/ChaosVehiclesPlugin/Source/ChaosVehicles/Public/ChaosVehicleWheel.h
@@ -290,6 +290,12 @@ class UChaosWheeledVehicleMovementComponent;
 		UFUNCTION(BlueprintCallable, Category = "Game|Components|WheeledVehicleMovement")
 		float GetSuspensionOffset() const;
 
+		UFUNCTION(BlueprintCallable, Category = "Game|Components|WheeledVehicleMovement")
+		float GetWheelRadius() const;
+
+		UFUNCTION(BlueprintCallable, Category = "Game|Components|WheeledVehicleMovement")
+		float GetWheelAngularVelocity() const;
+
 		UFUNCTION(BlueprintCallable, Category = "Game|Components|WheeledVehicleMovement")
 		bool IsInAir() const;
 
diff --git a/Engine/Plugins/Experimental/ChaosVehiclesPlugin/Source/ChaosVehicles/Public/ChaosWheeledVehicleMovementComponent.h b/Engine/Plugins/Experimental/ChaosVehiclesPlugin/Source/ChaosVehicles/Public/ChaosWheeledVehicleMovementComponent.h
--- a/Engine/Plugins/Experimental/ChaosVehiclesPlugin/Source/ChaosVehicles/Public/ChaosWheeledVehicleMovementComponent.h
+++ b/Engine/Plugins/Experimental/ChaosVehiclesPlugin/Source/ChaosVehicles/Public/ChaosWheeledVehicleMovementComponent.h
@@ -647,13 +647,29 @@ class CHAOSVEHICLES_API UChaosWheeledVehicleMovementComponent : public UChaosVeh
 	}
 
 	UFUNCTION(BlueprintPure, Category = "Vehicles")
-		static void BreakWheelStatus(const struct FWheelStatus& Status, bool& bInContact, FVector& ContactPoint, UPhysicalMaterial*& PhysMaterial
+	static void BreakWheelStatus(const struct FWheelStatus& Status, bool& bInContact, FVector& ContactPoint, UPhysicalMaterial*& PhysMaterial
 			, float& NormalizedSuspensionLength, float& SpringForce, bool& bIsSlipping, float& SlipMagnitude, bool& bIsSkidding, float& SkidMagnitude, FVector& SkidNormal);
 
 	UFUNCTION(BlueprintPure, Category = "Vehicles")
 	static FWheelStatus MakeWheelStatus(bool bInContact, FVector& ContactPoint, UPhysicalMaterial* PhysMaterial
 			, float NormalizedSuspensionLength, float SpringForce, bool bIsSlipping, float SlipMagnitude, bool bIsSkidding, float SkidMagnitude, FVector& SkidNormal);
 
+	UFUNCTION(BlueprintPure, Category = "Vehicles")
+	static void BreakWheeledSnapshot(const struct FWheeledSnaphotData& Snapshot, FTransform& Transform, FVector& LinearVelocity
+			, FVector& AngularVelocity, int& SelectedGear, float& EngineRPM, TArray<FWheelSnapshot>& WheelSnapshots);
+
+	UFUNCTION(BlueprintPure, Category = "Vehicles")
+	static FWheeledSnaphotData MakeWheeledSnapshot(FTransform Transform, FVector LinearVelocity, FVector AngularVelocity
+			, int SelectedGear, float EngineRPM, TArray<FWheelSnapshot>& WheelSnapshots);
+
+	UFUNCTION(BlueprintPure, Category = "Vehicles")
+	static void BreakWheelSnapshot(const struct FWheelSnapshot& Snapshot, float& SuspensionOffset
+			, float& WheelRotationAngle, float& SteeringAngle, float& WheelRadius, float& WheelAngularVelocity);
+
+	UFUNCTION(BlueprintPure, Category = "Vehicles")
+	static FWheelSnapshot MakeWheelSnapshot(float SuspensionOffset, float WheelRotationAngle
+			, float SteeringAngle, float WheelRadius, float WheelAngularVelocity);
+
 	/** Get a wheels current simulation state */
 	UFUNCTION(BlueprintCallable, Category = "Game|Components|ChaosWheeledVehicleMovement")
 	const FWheelStatus& GetWheelState(int WheelIndex) const
@@ -711,6 +727,14 @@ class CHAOSVEHICLES_API UChaosWheeledVehicleMovementComponent : public UChaosVeh
 	UFUNCTION(BlueprintCallable, Category = "Game|Components|ChaosWheeledVehicleMovement")
 	void SetWheelClass(int WheelIndex, TSubclassOf<UChaosVehicleWheel> InWheelClass);
 
+	/** Grab a snapshot of the vehicle instance dynamic state */
+	UFUNCTION(BlueprintCallable, Category = "Game|Components|ChaosWheeledVehicleMovement")
+	virtual FWheeledSnaphotData GetSnapshot() const;
+
+	/** Set snapshot of vehicle instance dynamic state */
+	UFUNCTION(BlueprintCallable, Category = "Game|Components|ChaosWheeledVehicleMovement")
+	virtual void SetSnapshot(const FWheeledSnaphotData& SnapshotIn);
+
 
 protected:
 
diff --git a/Engine/Plugins/Experimental/ChaosVehiclesPlugin/Source/ChaosVehicles/Public/SnapshotData.h b/Engine/Plugins/Experimental/ChaosVehiclesPlugin/Source/ChaosVehicles/Public/SnapshotData.h
new file mode 100644
--- /dev/null
+++ b/Engine/Plugins/Experimental/ChaosVehiclesPlugin/Source/ChaosVehicles/Public/SnapshotData.h
@@ -0,0 +1,62 @@
+// Copyright Epic Games, Inc. All Rights Reserved.
+
+#pragma once
+
+#include "CoreMinimal.h"
+#include "SnapshotData.generated.h"
+
+
+USTRUCT(BlueprintType)
+struct FBaseSnapshotData //: public UObject
+{
+public:
+	GENERATED_BODY()
+
+	UPROPERTY()
+	FTransform Transform;		// world coords
+
+	UPROPERTY()
+	FVector LinearVelocity;		// world coords
+
+	UPROPERTY()
+	FVector AngularVelocity;	// world coords
+
+};
+
+USTRUCT(BlueprintType)
+struct FWheelSnapshot
+{
+	GENERATED_BODY()
+	
+	UPROPERTY()
+	float SuspensionOffset;		// suspension location
+	
+	UPROPERTY()
+	float WheelRotationAngle;	// wheel rotation angle, rotated position
+	
+	UPROPERTY()
+	float SteeringAngle;		// steering position
+
+	UPROPERTY()
+	float WheelRadius;			// radius of the wheel can be changed dynamically, to sim damaged ot flat
+
+	UPROPERTY()
+	float WheelAngularVelocity;	// speed of rotation of wheel
+};
+
+USTRUCT(BlueprintType)
+struct FWheeledSnaphotData : public FBaseSnapshotData
+{
+public:
+	GENERATED_BODY()
+
+	UPROPERTY()
+	int SelectedGear;		// -ve reverse gear(s), 0 neutral, +ve forward gears
+
+	UPROPERTY()
+	float EngineRPM;		// Engine Revolutions Per Minute
+
+	UPROPERTY()
+	TArray<FWheelSnapshot> WheelSnapshots;
+
+};
diff --git a/Engine/Source/Runtime/Experimental/ChaosVehicles/ChaosVehiclesCore/Public/WheelSystem.h b/Engine/Source/Runtime/Experimental/ChaosVehicles/ChaosVehiclesCore/Public/WheelSystem.h
--- a/Engine/Source/Runtime/Experimental/ChaosVehicles/ChaosVehiclesCore/Public/WheelSystem.h
+++ b/Engine/Source/Runtime/Experimental/ChaosVehicles/ChaosVehiclesCore/Public/WheelSystem.h
@@ -151,6 +151,18 @@ public:
 		Re = NewRadius;
 	}
 
+	/** Set the angular position in radians */
+	void SetAngularPosition(float PositionIn)
+	{
+		AngularPosition = PositionIn;
+	}
+
+	/** Set the angular position in radians/sec */
+	void SetAngularVelocity(float AngularVelocityIn)
+	{
+		Omega = AngularVelocityIn;
+	}
+
 	/** set wheel rotational speed to match the specified linear forwards speed */
 	void SetMatchingSpeed(float LinearMetersPerSecondIn)
 	{
