ChaosVehicles:  Prevent code asserting when controls are applied to a vehicle with no engine setup

diff --git a/Engine/Plugins/Experimental/ChaosVehiclesPlugin/Source/ChaosVehicles/Private/ChaosWheeledVehicleMovementComponent.cpp b/Engine/Plugins/Experimental/ChaosVehiclesPlugin/Source/ChaosVehicles/Private/ChaosWheeledVehicleMovementComponent.cpp
--- a/Engine/Plugins/Experimental/ChaosVehiclesPlugin/Source/ChaosVehicles/Private/ChaosWheeledVehicleMovementComponent.cpp
+++ b/Engine/Plugins/Experimental/ChaosVehiclesPlugin/Source/ChaosVehicles/Private/ChaosWheeledVehicleMovementComponent.cpp
@@ -609,7 +609,7 @@ void UChaosWheeledVehicleSimulation::ApplyInput(const FControlInputs& ControlInp
 	FControlInputs ModifiedInputs = ControlInputs;
 
 	float EngineBraking = 0.f;
-	if (PVehicle->bMechanicalSimEnabled)
+	if (PVehicle->HasTransmission() && PVehicle->HasEngine())
 	{
 		auto& PEngine = PVehicle->GetEngine();
 		auto& PTransmission = PVehicle->GetTransmission();
