Fix CIS static analysis

diff --git a/Engine/Plugins/Experimental/ChaosVehiclesPlugin/Source/ChaosVehicles/Private/ChaosWheeledVehicleMovementComponent.cpp b/Engine/Plugins/Experimental/ChaosVehiclesPlugin/Source/ChaosVehicles/Private/ChaosWheeledVehicleMovementComponent.cpp
--- a/Engine/Plugins/Experimental/ChaosVehiclesPlugin/Source/ChaosVehicles/Private/ChaosWheeledVehicleMovementComponent.cpp
+++ b/Engine/Plugins/Experimental/ChaosVehiclesPlugin/Source/ChaosVehicles/Private/ChaosWheeledVehicleMovementComponent.cpp
@@ -1233,7 +1233,6 @@ void UChaosWheeledVehicleMovementComponent::SetupVehicle(TUniquePtr<Chaos::FSimp
 	// cache this value as it's useful for steering setup calculations and debug rendering
 	WheelTrackDimensions = CalculateWheelLayoutDimensions();
 
-	if (bMechanicalSimEnabled)
 	if (bMechanicalSimEnabled)
 	{
 		Chaos::FSimpleEngineSim EngineSim(&EngineSetup.GetPhysicsEngineConfig());
