ChaosVehicles: Reinstate steering angle vs speed curve

diff --git a/Engine/Plugins/Experimental/ChaosVehiclesPlugin/Source/ChaosVehicles/Private/ChaosWheeledVehicleMovementComponent.cpp b/Engine/Plugins/Experimental/ChaosVehiclesPlugin/Source/ChaosVehicles/Private/ChaosWheeledVehicleMovementComponent.cpp
--- a/Engine/Plugins/Experimental/ChaosVehiclesPlugin/Source/ChaosVehicles/Private/ChaosWheeledVehicleMovementComponent.cpp
+++ b/Engine/Plugins/Experimental/ChaosVehiclesPlugin/Source/ChaosVehicles/Private/ChaosWheeledVehicleMovementComponent.cpp
@@ -569,16 +569,13 @@ void UChaosWheeledVehicleSimulation::ProcessSteering(const FControlInputs& Contr
 
 		if (PWheel.Setup().SteeringEnabled)
 		{
-			// #TODO: PUT BACK
-			//FRichCurve* SteeringCurveData = SteeringSetup.SteeringCurve.GetRichCurve();
 			float SpeedScale = 1.0f;
 
-			// #TODO: PUT BACK
-			//// allow full counter steering when steering into a power slide
-			////if (SteeringInput * VehicleState.VehicleLocalVelocity.Y > 0.0f)
-			//{
-			//	SpeedScale = SteeringCurveData->Eval(CmSToMPH(VehicleState.ForwardSpeed));
-			//}
+			// allow full counter steering when steering into a power slide
+			//if (ControlInputs.SteeringInput * VehicleState.VehicleLocalVelocity.Y > 0.0f)
+			{
+				SpeedScale = PVehicle->GetSteering().GetSteeringFromVelocity(CmSToMPH(VehicleState.ForwardSpeed));
+			}
 
 			float UseSteeringValue = ControlInputs.SteeringInput * SpeedScale;
 
diff --git a/Engine/Plugins/Experimental/ChaosVehiclesPlugin/Source/ChaosVehicles/Public/ChaosWheeledVehicleMovementComponent.h b/Engine/Plugins/Experimental/ChaosVehiclesPlugin/Source/ChaosVehicles/Public/ChaosWheeledVehicleMovementComponent.h
--- a/Engine/Plugins/Experimental/ChaosVehiclesPlugin/Source/ChaosVehicles/Public/ChaosWheeledVehicleMovementComponent.h
+++ b/Engine/Plugins/Experimental/ChaosVehiclesPlugin/Source/ChaosVehicles/Public/ChaosWheeledVehicleMovementComponent.h
@@ -463,10 +463,10 @@ private:
 		float MaxX = this->SteeringCurve.GetRichCurveConst()->GetLastKey().Time;
 		PSteeringConfig.SpeedVsSteeringCurve.Empty();
 		float NumSamples = 20;
-		for (float X = 0; X <= MaxX; X += (MaxX / NumSamples))
+		for (float X = MinValue; X <= MaxX; X += ((MaxX- MinValue) / NumSamples))
 		{
 			float Y = this->SteeringCurve.GetRichCurveConst()->Eval(X) / MaxValue;
-			PSteeringConfig.SpeedVsSteeringCurve.AddNormalized(Y);
+			PSteeringConfig.SpeedVsSteeringCurve.Add(FVector2D(X, Y));
 		}
 
 		PSteeringConfig.TrackWidth = WheelTrackDimensions.Y;
diff --git a/Engine/Source/Runtime/Experimental/ChaosVehicles/ChaosVehiclesCore/Private/VehicleUtility.cpp b/Engine/Source/Runtime/Experimental/ChaosVehicles/ChaosVehiclesCore/Private/VehicleUtility.cpp
--- a/Engine/Source/Runtime/Experimental/ChaosVehicles/ChaosVehiclesCore/Private/VehicleUtility.cpp
+++ b/Engine/Source/Runtime/Experimental/ChaosVehicles/ChaosVehiclesCore/Private/VehicleUtility.cpp
@@ -166,6 +166,48 @@ namespace Chaos
 		return Radius;
 	}
 
+	void FGraph::Add(const FVec2& Value)
+	{
+		Graph.Add(Value);
+
+		if (Value.X < BoundsX.X)
+		{
+			BoundsX.X = Value.X;
+		}
+		if (Value.X > BoundsX.Y)
+		{
+			BoundsX.Y = Value.X;
+		}
+		if (Value.X < BoundsY.X)
+		{
+			BoundsY.X = Value.X;
+		}
+		if (Value.X > BoundsY.Y)
+		{
+			BoundsY.Y = Value.X;
+		}
+	}
+
+	float FGraph::EvaluateY(float InX) const
+	{
+		float RangeX = BoundsX.Y - BoundsX.X;
+		float Step = (RangeX) / (Graph.Num() - 1);
+		int StartIndex = (InX - BoundsX.X) / Step;
+
+		if (StartIndex <= 0)
+		{
+			return Graph[0].Y;
+		}
+		else if (StartIndex >= Graph.Num() - 1)
+		{
+			return Graph[Graph.Num() - 1].Y;
+		}
+
+		float NormalisedRamp = (InX - Graph[StartIndex].X) / Step;
+
+		return  Graph[StartIndex].Y * (1.f - NormalisedRamp) + Graph[StartIndex + 1].Y * NormalisedRamp;
+	}
+
 } // namespace Chaos
 
 #if VEHICLE_DEBUGGING_ENABLED
diff --git a/Engine/Source/Runtime/Experimental/ChaosVehicles/ChaosVehiclesCore/Public/SteeringSystem.h b/Engine/Source/Runtime/Experimental/ChaosVehicles/ChaosVehiclesCore/Public/SteeringSystem.h
--- a/Engine/Source/Runtime/Experimental/ChaosVehicles/ChaosVehiclesCore/Public/SteeringSystem.h
+++ b/Engine/Source/Runtime/Experimental/ChaosVehicles/ChaosVehiclesCore/Public/SteeringSystem.h
@@ -37,7 +37,7 @@ namespace Chaos
 
 		float MaxSteeringAngle;
 
-		FNormalisedGraph SpeedVsSteeringCurve;
+		FGraph SpeedVsSteeringCurve;
 	};
 
 	class CHAOSVEHICLESCORE_API FAckermannSim : public TVehicleSystem<FSimpleSteeringConfig>
@@ -82,6 +82,11 @@ namespace Chaos
 		{
 		}
 
+		float GetSteeringFromVelocity(float VelocityMPH)
+		{
+			return Setup().SpeedVsSteeringCurve.EvaluateY(VelocityMPH);
+		}
+
 		float GetSteeringAngle(float InNormSteering, float MaxSteeringAngle, float WheelSide)
 		{
 			float OutSteeringAngle = 0.f;
@@ -118,7 +123,6 @@ namespace Chaos
 
 			return OutSteeringAngle;
 		}
-
 		FAckermannSim Ackermann;
 	};
 
diff --git a/Engine/Source/Runtime/Experimental/ChaosVehicles/ChaosVehiclesCore/Public/VehicleUtility.h b/Engine/Source/Runtime/Experimental/ChaosVehicles/ChaosVehiclesCore/Public/VehicleUtility.h
--- a/Engine/Source/Runtime/Experimental/ChaosVehicles/ChaosVehiclesCore/Public/VehicleUtility.h
+++ b/Engine/Source/Runtime/Experimental/ChaosVehicles/ChaosVehiclesCore/Public/VehicleUtility.h
@@ -3,6 +3,7 @@
 #pragma once
 
 #include "CoreMinimal.h"
+#include "Chaos/Core.h"
 
 // Disable Optimizations in non debug build configurations
 #define VEHICLE_DEBUGGING_ENABLED 0
@@ -69,6 +70,32 @@ namespace Chaos
 		TArray<float> Graph;
 	};
 
+	class CHAOSVEHICLESCORE_API FGraph
+	{
+	public:
+		FGraph()
+		{
+			Empty();
+		}
+
+		void Empty()
+		{
+			Graph.Empty();
+			BoundsX.X = TNumericLimits<FReal>::Max();
+			BoundsX.Y = -TNumericLimits<FReal>::Max();
+			BoundsY.X = TNumericLimits<FReal>::Max();
+			BoundsY.Y = -TNumericLimits<FReal>::Max();
+		}
+
+		void Add(const FVec2& Value);
+
+		float EvaluateY(float InX) const;
+	private:
+		TArray<FVec2> Graph;
+		FVector2D BoundsX;
+		FVector2D BoundsY;
+	};
+
 	class CHAOSVEHICLESCORE_API FVehicleUtility
 	{
 	public:
@@ -309,4 +336,3 @@ namespace Chaos
 	};
 
 } // namespace Chaos
-
