ChaosVehicle: Steering Improvements
* option to use non-linear input scaling (Squared relationship)
* Steering range vs speed curve with speed now working

diff --git a/Engine/Plugins/Experimental/ChaosVehiclesPlugin/Source/ChaosVehicles/Private/ChaosVehicleMovementComponent.cpp b/Engine/Plugins/Experimental/ChaosVehiclesPlugin/Source/ChaosVehicles/Private/ChaosVehicleMovementComponent.cpp
--- a/Engine/Plugins/Experimental/ChaosVehiclesPlugin/Source/ChaosVehicles/Private/ChaosVehicleMovementComponent.cpp
+++ b/Engine/Plugins/Experimental/ChaosVehiclesPlugin/Source/ChaosVehicles/Private/ChaosVehicleMovementComponent.cpp
@@ -628,18 +628,24 @@ UChaosVehicleMovementComponent::UChaosVehicleMovementComponent(const FObjectInit
 	WrongDirectionThreshold = 100.f;
 	ThrottleInputRate.RiseRate = 6.0f;
 	ThrottleInputRate.FallRate = 10.0f;
+	ThrottleInputRate.InputCurveFunction = EInputFunctionType::LinearFunction;
 	BrakeInputRate.RiseRate = 6.0f;
 	BrakeInputRate.FallRate = 10.0f;
+	BrakeInputRate.InputCurveFunction = EInputFunctionType::LinearFunction;
 	SteeringInputRate.RiseRate = 2.5f;
 	SteeringInputRate.FallRate = 5.0f;
+	SteeringInputRate.InputCurveFunction = EInputFunctionType::SquaredFunction;
 	HandbrakeInputRate.RiseRate = 12.0f;
 	HandbrakeInputRate.FallRate = 12.0f;
 	PitchInputRate.RiseRate = 6.0f;
 	PitchInputRate.FallRate = 10.0f;
+	PitchInputRate.InputCurveFunction = EInputFunctionType::LinearFunction;
 	RollInputRate.RiseRate = 6.0f;
 	RollInputRate.FallRate = 10.0f;
+	RollInputRate.InputCurveFunction = EInputFunctionType::LinearFunction;
 	YawInputRate.RiseRate = 6.0f;
 	YawInputRate.FallRate = 10.0f;
+	YawInputRate.InputCurveFunction = EInputFunctionType::LinearFunction;
 	TransmissionType = Chaos::ETransmissionType::Automatic;
 
 	SetIsReplicatedByDefault(true);
@@ -1666,13 +1672,13 @@ void UChaosVehicleMovementComponent::Update(float DeltaTime)
 				CurAsyncInput->Actor.Proxy = Handle;	// vehicles are never static
 
 				FChaosVehicleDefaultAsyncInput* AsyncInput = static_cast<FChaosVehicleDefaultAsyncInput*>(CurAsyncInput);
-				AsyncInput->ControlInputs.ThrottleInput = ThrottleInput;
-				AsyncInput->ControlInputs.BrakeInput = BrakeInput;
-				AsyncInput->ControlInputs.SteeringInput = SteeringInput;
+				AsyncInput->ControlInputs.ThrottleInput = ThrottleInputRate.CalcControlFunction(ThrottleInput);
+				AsyncInput->ControlInputs.BrakeInput = BrakeInputRate.CalcControlFunction(BrakeInput);
+				AsyncInput->ControlInputs.SteeringInput = SteeringInputRate.CalcControlFunction(SteeringInput);
 				AsyncInput->ControlInputs.HandbrakeInput = HandbrakeInput;
-				AsyncInput->ControlInputs.RollInput = RollInput;
-				AsyncInput->ControlInputs.PitchInput = PitchInput;
-				AsyncInput->ControlInputs.YawInput = YawInput;
+				AsyncInput->ControlInputs.RollInput = RollInputRate.CalcControlFunction(RollInput);
+				AsyncInput->ControlInputs.PitchInput = PitchInputRate.CalcControlFunction(PitchInput);
+				AsyncInput->ControlInputs.YawInput = YawInputRate.CalcControlFunction(YawInput);
 				AsyncInput->ControlInputs.GearUpInput = bRawGearUpInput;
 				AsyncInput->ControlInputs.GearDownInput = bRawGearDownInput;
 				AsyncInput->ControlInputs.TransmissionType = TransmissionType;
diff --git a/Engine/Plugins/Experimental/ChaosVehiclesPlugin/Source/ChaosVehicles/Private/ChaosWheeledVehicleMovementComponent.cpp b/Engine/Plugins/Experimental/ChaosVehiclesPlugin/Source/ChaosVehicles/Private/ChaosWheeledVehicleMovementComponent.cpp
--- a/Engine/Plugins/Experimental/ChaosVehiclesPlugin/Source/ChaosVehicles/Private/ChaosWheeledVehicleMovementComponent.cpp
+++ b/Engine/Plugins/Experimental/ChaosVehiclesPlugin/Source/ChaosVehicles/Private/ChaosWheeledVehicleMovementComponent.cpp
@@ -577,9 +577,8 @@ void UChaosWheeledVehicleSimulation::ProcessSteering(const FControlInputs& Contr
 				SpeedScale = PVehicle->GetSteering().GetSteeringFromVelocity(CmSToMPH(VehicleState.ForwardSpeed));
 			}
 
-			float UseSteeringValue = ControlInputs.SteeringInput * SpeedScale;
+			float SteeringAngle = ControlInputs.SteeringInput * SpeedScale;
 
-			float SteeringAngle = 0.f;
 			if (FMath::Abs(GWheeledVehicleDebugParams.SteeringOverride) > 0.01f)
 			{
 				SteeringAngle = PWheel.Setup().MaxSteeringAngle * GWheeledVehicleDebugParams.SteeringOverride;
@@ -587,7 +586,7 @@ void UChaosWheeledVehicleSimulation::ProcessSteering(const FControlInputs& Contr
 			else
 			{
 				float WheelSide = PVehicle->GetSuspension(WheelIdx).GetLocalRestingPosition().Y;
-				SteeringAngle = PSteering.GetSteeringAngle(UseSteeringValue, PWheel.Setup().MaxSteeringAngle, WheelSide);
+				SteeringAngle = PSteering.GetSteeringAngle(SteeringAngle, PWheel.Setup().MaxSteeringAngle, WheelSide);
 			}
 
 			PWheel.SetSteeringAngle(SteeringAngle);
diff --git a/Engine/Plugins/Experimental/ChaosVehiclesPlugin/Source/ChaosVehicles/Public/ChaosVehicleMovementComponent.h b/Engine/Plugins/Experimental/ChaosVehiclesPlugin/Source/ChaosVehicles/Public/ChaosVehicleMovementComponent.h
--- a/Engine/Plugins/Experimental/ChaosVehiclesPlugin/Source/ChaosVehicles/Public/ChaosVehicleMovementComponent.h
+++ b/Engine/Plugins/Experimental/ChaosVehiclesPlugin/Source/ChaosVehicles/Public/ChaosVehicleMovementComponent.h
@@ -418,6 +418,14 @@ struct FVehicleState
 	int SleepCounter;
 };
 
+/** Input Options */
+UENUM()
+enum class EInputFunctionType : uint8
+{
+	LinearFunction,
+	SquaredFunction
+};
+
 USTRUCT()
 struct CHAOSVEHICLES_API FVehicleInputRateConfig
 {
@@ -431,6 +439,9 @@ struct CHAOSVEHICLES_API FVehicleInputRateConfig
 	UPROPERTY(EditAnywhere, Category=VehicleInputRate)
 	float FallRate;
 
+	UPROPERTY(EditAnywhere, Category=VehicleInputRate)
+	EInputFunctionType InputCurveFunction;
+
 	FVehicleInputRateConfig() : RiseRate(5.0f), FallRate(5.0f) { }
 
 	/** Change an output value using max rise and fall rates */
@@ -447,6 +458,28 @@ struct CHAOSVEHICLES_API FVehicleInputRateConfig
 		const float ClampedDeltaValue = FMath::Clamp( DeltaValue, -MaxDeltaValue, MaxDeltaValue );
 		return CurrentValue + ClampedDeltaValue;
 	}
+
+	float CalcControlFunction(float InputValue)
+	{
+		switch (InputCurveFunction)
+		{
+
+		case EInputFunctionType::SquaredFunction:
+		{
+			return (InputValue < 0.f) ? -InputValue * InputValue : InputValue * InputValue;
+		}
+		break;
+
+		case EInputFunctionType::LinearFunction:
+		default:
+		{
+			return InputValue;
+		}
+		break;
+		
+		}
+
+	}
 };
 
 
diff --git a/Engine/Plugins/Experimental/ChaosVehiclesPlugin/Source/ChaosVehicles/Public/ChaosWheeledVehicleMovementComponent.h b/Engine/Plugins/Experimental/ChaosVehiclesPlugin/Source/ChaosVehicles/Public/ChaosWheeledVehicleMovementComponent.h
--- a/Engine/Plugins/Experimental/ChaosVehiclesPlugin/Source/ChaosVehicles/Public/ChaosWheeledVehicleMovementComponent.h
+++ b/Engine/Plugins/Experimental/ChaosVehiclesPlugin/Source/ChaosVehicles/Public/ChaosWheeledVehicleMovementComponent.h
@@ -411,11 +411,15 @@ enum class ESteeringType : uint8
 	Ackermann,
 };
 
+
 USTRUCT()
 struct FVehicleSteeringConfig
 {
 	GENERATED_USTRUCT_BODY()
 
+	UPROPERTY(EditAnywhere, Category = SteeringSetup)
+	EInputFunctionType SteeringInputType;
+
 	/** Single angle : both wheels steer by the same amount
 	 *  AngleRatio   : outer wheels on corner steer less than the inner ones by set ratio 
 	 *  Ackermann	 : Ackermann steering principle is applied */
