ChaosVehicles: Debug mode to set a top speed for the vehicle.

diff --git a/Engine/Plugins/Experimental/ChaosVehiclesPlugin/Source/ChaosVehicles/Private/ChaosVehicleMovementComponent.cpp b/Engine/Plugins/Experimental/ChaosVehiclesPlugin/Source/ChaosVehicles/Private/ChaosVehicleMovementComponent.cpp
--- a/Engine/Plugins/Experimental/ChaosVehiclesPlugin/Source/ChaosVehicles/Private/ChaosVehicleMovementComponent.cpp
+++ b/Engine/Plugins/Experimental/ChaosVehiclesPlugin/Source/ChaosVehicles/Private/ChaosVehicleMovementComponent.cpp
@@ -62,6 +62,7 @@ FAutoConsoleVariableRef CVarChaosVehiclesBatchQueries(TEXT("p.Vehicle.BatchQueri
 FAutoConsoleVariableRef CVarChaosVehiclesForceDebugScaling(TEXT("p.Vehicle.SetForceDebugScaling"), GVehicleDebugParams.ForceDebugScaling, TEXT("Set Scaling For Force Visualisation."));
 FAutoConsoleVariableRef CVarChaosVehiclesSleepCounterThreshold(TEXT("p.Vehicle.SleepCounterThreshold"), GVehicleDebugParams.SleepCounterThreshold, TEXT("Set The Sleep Counter Iteration Threshold."));
 FAutoConsoleVariableRef CVarChaosVehiclesDisableVehicleSleep(TEXT("p.Vehicle.DisableVehicleSleep"), GVehicleDebugParams.DisableVehicleSleep, TEXT("Disable Vehicle Agressive Sleeping."));
+FAutoConsoleVariableRef CVarChaosVehiclesSetMaxMPH(TEXT("p.Vehicle.SetMaxMPH"), GVehicleDebugParams.SetMaxMPH, TEXT("Set a top speed in MPH (affects all vehicles)."));
 
 
 void FVehicleState::CaptureState(const FBodyInstance* TargetInstance, float GravityZ, float DeltaTime)
@@ -1685,6 +1686,12 @@ void UChaosVehicleMovementComponent::Update(float DeltaTime)
 
 				AsyncInput->ControlInputs.ParkingEnabled = bParkEnabled;
 
+				// debug feature to limit the vehicles top speed
+				if ((GVehicleDebugParams.SetMaxMPH > 0.f) && (FMath::Abs(ThrottleInput) > 0.0f) && FMath::Abs(GetForwardSpeedMPH()) >= GVehicleDebugParams.SetMaxMPH)
+				{
+					AsyncInput->ControlInputs.ThrottleInput = 0.1f;
+				}
+
 				AsyncInput->GravityZ = GetGravityZ();
 
 				// #TODO: don't create every frame, setup never changes dynamically
diff --git a/Engine/Plugins/Experimental/ChaosVehiclesPlugin/Source/ChaosVehicles/Public/ChaosVehicleMovementComponent.h b/Engine/Plugins/Experimental/ChaosVehiclesPlugin/Source/ChaosVehicles/Public/ChaosVehicleMovementComponent.h
--- a/Engine/Plugins/Experimental/ChaosVehiclesPlugin/Source/ChaosVehicles/Public/ChaosVehicleMovementComponent.h
+++ b/Engine/Plugins/Experimental/ChaosVehiclesPlugin/Source/ChaosVehicles/Public/ChaosVehicleMovementComponent.h
@@ -94,6 +94,7 @@ struct FVehicleDebugParams
 	float SleepCounterThreshold = 15;
 	bool DisableVehicleSleep = false;
 	bool EnableMultithreading = false; // possibly not all entirely thread safe at present
+	float SetMaxMPH = 0.0f;
 };
 
 struct FBodyInstance;
