ChaosVehicles: Fix Blueprint functionality, input appearing as output because it's not const

diff --git a/Engine/Plugins/Experimental/ChaosVehiclesPlugin/Source/ChaosVehicles/Private/ChaosWheeledVehicleMovementComponent.cpp b/Engine/Plugins/Experimental/ChaosVehiclesPlugin/Source/ChaosVehicles/Private/ChaosWheeledVehicleMovementComponent.cpp
--- a/Engine/Plugins/Experimental/ChaosVehiclesPlugin/Source/ChaosVehicles/Private/ChaosWheeledVehicleMovementComponent.cpp
+++ b/Engine/Plugins/Experimental/ChaosVehiclesPlugin/Source/ChaosVehicles/Private/ChaosWheeledVehicleMovementComponent.cpp
@@ -1928,7 +1928,7 @@ void UChaosWheeledVehicleMovementComponent::BreakWheeledSnapshot(const struct FW
 }
 
 FWheeledSnaphotData UChaosWheeledVehicleMovementComponent::MakeWheeledSnapshot(FTransform Transform, FVector LinearVelocity
-	, FVector AngularVelocity, int SelectedGear, float EngineRPM, TArray<FWheelSnapshot>& WheelSnapshots)
+	, FVector AngularVelocity, int SelectedGear, float EngineRPM, const TArray<FWheelSnapshot>& WheelSnapshots)
 {
 	FWheeledSnaphotData Snapshot;
 	Snapshot.Transform = Transform;
diff --git a/Engine/Plugins/Experimental/ChaosVehiclesPlugin/Source/ChaosVehicles/Public/ChaosWheeledVehicleMovementComponent.h b/Engine/Plugins/Experimental/ChaosVehiclesPlugin/Source/ChaosVehicles/Public/ChaosWheeledVehicleMovementComponent.h
--- a/Engine/Plugins/Experimental/ChaosVehiclesPlugin/Source/ChaosVehicles/Public/ChaosWheeledVehicleMovementComponent.h
+++ b/Engine/Plugins/Experimental/ChaosVehiclesPlugin/Source/ChaosVehicles/Public/ChaosWheeledVehicleMovementComponent.h
@@ -664,7 +664,7 @@ class CHAOSVEHICLES_API UChaosWheeledVehicleMovementComponent : public UChaosVeh
 
 	UFUNCTION(BlueprintPure, Category = "Vehicles")
 	static FWheeledSnaphotData MakeWheeledSnapshot(FTransform Transform, FVector LinearVelocity, FVector AngularVelocity
-			, int SelectedGear, float EngineRPM, TArray<FWheelSnapshot>& WheelSnapshots);
+			, int SelectedGear, float EngineRPM, const TArray<FWheelSnapshot>& WheelSnapshots);
 
 	UFUNCTION(BlueprintPure, Category = "Vehicles")
 	static void BreakWheelSnapshot(const struct FWheelSnapshot& Snapshot, float& SuspensionOffset
