ChaosVehicles: Add optional input curve per control

diff --git a/Engine/Plugins/Experimental/ChaosVehiclesPlugin/Source/ChaosVehicles/Public/ChaosVehicleMovementComponent.h b/Engine/Plugins/Experimental/ChaosVehiclesPlugin/Source/ChaosVehicles/Public/ChaosVehicleMovementComponent.h
--- a/Engine/Plugins/Experimental/ChaosVehiclesPlugin/Source/ChaosVehicles/Public/ChaosVehicleMovementComponent.h
+++ b/Engine/Plugins/Experimental/ChaosVehiclesPlugin/Source/ChaosVehicles/Public/ChaosVehicleMovementComponent.h
@@ -424,7 +424,8 @@ UENUM()
 enum class EInputFunctionType : uint8
 {
 	LinearFunction,
-	SquaredFunction
+	SquaredFunction,
+	CustomCurve
 };
 
 USTRUCT()
@@ -432,17 +433,31 @@ struct CHAOSVEHICLES_API FVehicleInputRateConfig
 {
 	GENERATED_USTRUCT_BODY()
 
-	// Rate at which the input value rises
+	/** 
+	 * Rate at which the input value rises
+	 */
 	UPROPERTY(EditAnywhere, Category=VehicleInputRate)
 	float RiseRate;
 
-	// Rate at which the input value falls
+	/**
+	 * Rate at which the input value falls
+	 */
 	UPROPERTY(EditAnywhere, Category=VehicleInputRate)
 	float FallRate;
 
+	/**
+	 * Controller input curve, various predefined options, linear, squared, or user can specify a custom curve function
+	 */
 	UPROPERTY(EditAnywhere, Category=VehicleInputRate)
 	EInputFunctionType InputCurveFunction;
 
+	/**
+	 * Controller input curve - should be a normalized float curve, i.e. time from 0 to 1 and values between 0 and 1
+	 * This curve is only sued if the InputCurveFunction above is set to CustomCurve
+	 */
+	UPROPERTY(EditAnywhere, Category = VehicleInputRate)
+	FRuntimeFloatCurve UserCurve;
+
 	FVehicleInputRateConfig() : RiseRate(5.0f), FallRate(5.0f) { }
 
 	/** Change an output value using max rise and fall rates */
@@ -462,9 +477,24 @@ struct CHAOSVEHICLES_API FVehicleInputRateConfig
 
 	float CalcControlFunction(float InputValue)
 	{
+		// user defined curve
+
+		// else use option from drop down list
 		switch (InputCurveFunction)
 		{
-
+		case EInputFunctionType::CustomCurve:
+		{
+			if (UserCurve.GetRichCurveConst() && !UserCurve.GetRichCurveConst()->IsEmpty())
+			{
+				float Output = UserCurve.GetRichCurveConst()->Eval(InputValue);
+				return FMath::Clamp(Output, 0.0f, 1.0f); // Ensure output is normalized
+			}
+			else
+			{
+				return InputValue;
+			}
+		}
+		break;
 		case EInputFunctionType::SquaredFunction:
 		{
 			return (InputValue < 0.f) ? -InputValue * InputValue : InputValue * InputValue;
diff --git a/Engine/Plugins/Experimental/ChaosVehiclesPlugin/Source/ChaosVehicles/Public/ChaosWheeledVehicleMovementComponent.h b/Engine/Plugins/Experimental/ChaosVehiclesPlugin/Source/ChaosVehicles/Public/ChaosWheeledVehicleMovementComponent.h
--- a/Engine/Plugins/Experimental/ChaosVehiclesPlugin/Source/ChaosVehicles/Public/ChaosWheeledVehicleMovementComponent.h
+++ b/Engine/Plugins/Experimental/ChaosVehiclesPlugin/Source/ChaosVehicles/Public/ChaosWheeledVehicleMovementComponent.h
@@ -417,9 +417,6 @@ struct FVehicleSteeringConfig
 {
 	GENERATED_USTRUCT_BODY()
 
-	UPROPERTY(EditAnywhere, Category = SteeringSetup)
-	EInputFunctionType SteeringInputType;
-
 	/** Single angle : both wheels steer by the same amount
 	 *  AngleRatio   : outer wheels on corner steer less than the inner ones by set ratio 
 	 *  Ackermann	 : Ackermann steering principle is applied */
