Unshelved from pending changelist '15378986':

ChaosVehicles: Fix raytrace complex/simple wheel setup parameter

diff --git a/Engine/Plugins/Experimental/ChaosVehiclesPlugin/Source/ChaosVehicles/Private/ChaosWheeledVehicleMovementComponent.cpp b/Engine/Plugins/Experimental/ChaosVehiclesPlugin/Source/ChaosVehicles/Private/ChaosWheeledVehicleMovementComponent.cpp
--- a/Engine/Plugins/Experimental/ChaosVehiclesPlugin/Source/ChaosVehicles/Private/ChaosWheeledVehicleMovementComponent.cpp
+++ b/Engine/Plugins/Experimental/ChaosVehiclesPlugin/Source/ChaosVehicles/Private/ChaosWheeledVehicleMovementComponent.cpp
@@ -229,7 +229,7 @@ void UChaosWheeledVehicleSimulation::UpdateSimulation(float DeltaTime, const FCh
 
 }
 
-void UChaosWheeledVehicleSimulation::PerformSuspensionTraces(const TArray<FSuspensionTrace>& SuspensionTrace, const FCollisionQueryParams& TraceParams)
+void UChaosWheeledVehicleSimulation::PerformSuspensionTraces(const TArray<FSuspensionTrace>& SuspensionTrace, FCollisionQueryParams& TraceParams)
 {
 	//#todo: SpringCollisionChannel should be a parameter setup
 	ECollisionChannel SpringCollisionChannel = ECollisionChannel::ECC_WorldDynamic;
@@ -282,7 +282,7 @@ void UChaosWheeledVehicleSimulation::PerformSuspensionTraces(const TArray<FSuspe
 			{
 				const FVector& TraceStart = SuspensionTrace[WheelIdx].Start;
 				const FVector& TraceEnd = SuspensionTrace[WheelIdx].End;
-				//TraceParams.bTraceComplex = (Wheels[WheelIdx]->SweepType == ESweepType::ComplexSweep);
+				TraceParams.bTraceComplex = (Wheels[WheelIdx]->SweepType == ESweepType::ComplexSweep);
 
 				FVector TraceVector(TraceStart - TraceEnd); // reversed
 				FVector TraceNormal = TraceVector.GetSafeNormal();
@@ -342,7 +342,7 @@ void UChaosWheeledVehicleSimulation::PerformSuspensionTraces(const TArray<FSuspe
 
 			FVector TraceStart = SuspensionTrace[WheelIdx].Start;
 			FVector TraceEnd = SuspensionTrace[WheelIdx].End;
-			//TraceParams.bTraceComplex = (Wheels[WheelIdx]->SweepType == ESweepType::ComplexSweep);
+			TraceParams.bTraceComplex = (Wheels[WheelIdx]->SweepType == ESweepType::ComplexSweep);
 
 			FVector TraceVector(TraceStart - TraceEnd); // reversed
 			FVector TraceNormal = TraceVector.GetSafeNormal();
diff --git a/Engine/Plugins/Experimental/ChaosVehiclesPlugin/Source/ChaosVehicles/Public/ChaosVehicleMovementComponent.h b/Engine/Plugins/Experimental/ChaosVehiclesPlugin/Source/ChaosVehicles/Public/ChaosVehicleMovementComponent.h
--- a/Engine/Plugins/Experimental/ChaosVehiclesPlugin/Source/ChaosVehicles/Public/ChaosVehicleMovementComponent.h
+++ b/Engine/Plugins/Experimental/ChaosVehiclesPlugin/Source/ChaosVehicles/Public/ChaosVehicleMovementComponent.h
@@ -150,7 +150,7 @@ struct FChaosVehicleDefaultAsyncInput : public FChaosVehicleAsyncInput
 {
 	float GravityZ;
 	FControlInputs ControlInputs;
-	FCollisionQueryParams TraceParams;
+	mutable FCollisionQueryParams TraceParams;
 
 	FChaosVehicleDefaultAsyncInput();
 
diff --git a/Engine/Plugins/Experimental/ChaosVehiclesPlugin/Source/ChaosVehicles/Public/ChaosWheeledVehicleMovementComponent.h b/Engine/Plugins/Experimental/ChaosVehiclesPlugin/Source/ChaosVehicles/Public/ChaosWheeledVehicleMovementComponent.h
--- a/Engine/Plugins/Experimental/ChaosVehiclesPlugin/Source/ChaosVehicles/Public/ChaosWheeledVehicleMovementComponent.h
+++ b/Engine/Plugins/Experimental/ChaosVehiclesPlugin/Source/ChaosVehicles/Public/ChaosWheeledVehicleMovementComponent.h
@@ -563,7 +563,7 @@ public:
 	virtual void ApplyInput(const FControlInputs& ControlInputs, float DeltaTime) override;
 
 	/** Perform suspension ray/shape traces */
-	virtual void PerformSuspensionTraces(const TArray<Chaos::FSuspensionTrace>& SuspensionTrace, const FCollisionQueryParams& TraceParams);
+	virtual void PerformSuspensionTraces(const TArray<Chaos::FSuspensionTrace>& SuspensionTrace, FCollisionQueryParams& TraceParams);
 
 
 	/** Update the engine/transmission simulation */
