ChaosVehicles: Fix assert when using sequencer

diff --git a/Engine/Plugins/Experimental/ChaosVehiclesPlugin/Source/ChaosVehicles/Private/ChaosWheeledVehicleMovementComponent.cpp b/Engine/Plugins/Experimental/ChaosVehiclesPlugin/Source/ChaosVehicles/Private/ChaosWheeledVehicleMovementComponent.cpp
--- a/Engine/Plugins/Experimental/ChaosVehiclesPlugin/Source/ChaosVehicles/Private/ChaosWheeledVehicleMovementComponent.cpp
+++ b/Engine/Plugins/Experimental/ChaosVehiclesPlugin/Source/ChaosVehicles/Private/ChaosWheeledVehicleMovementComponent.cpp
@@ -686,56 +686,59 @@ bool UChaosWheeledVehicleSimulation::IsWheelSpinning() const
 
 void UChaosWheeledVehicleSimulation::ProcessMechanicalSimulation(float DeltaTime)
 {
-	auto& PEngine = PVehicle->GetEngine();
-	auto& PTransmission = PVehicle->GetTransmission();
-	auto& PDifferential = PVehicle->GetDifferential();
-
-	float WheelRPM = 0;
-	for (int I = 0; I < PVehicle->Wheels.Num(); I++)
+	if (PVehicle->HasEngine())
 	{
-		if (PVehicle->Wheels[I].Setup().EngineEnabled)
+		auto& PEngine = PVehicle->GetEngine();
+		auto& PTransmission = PVehicle->GetTransmission();
+		auto& PDifferential = PVehicle->GetDifferential();
+
+		float WheelRPM = 0;
+		for (int I = 0; I < PVehicle->Wheels.Num(); I++)
 		{
-			WheelRPM = FMath::Abs(PVehicle->Wheels[I].GetWheelRPM());
+			if (PVehicle->Wheels[I].Setup().EngineEnabled)
+			{
+				WheelRPM = FMath::Abs(PVehicle->Wheels[I].GetWheelRPM());
+			}
 		}
-	}
 
-	PEngine.SetEngineRPM(PTransmission.IsOutOfGear(), PTransmission.GetEngineRPMFromWheelRPM(WheelRPM));
-	PEngine.Simulate(DeltaTime);
+		PEngine.SetEngineRPM(PTransmission.IsOutOfGear(), PTransmission.GetEngineRPMFromWheelRPM(WheelRPM));
+		PEngine.Simulate(DeltaTime);
 
-	PTransmission.SetEngineRPM(PEngine.GetEngineRPM()); // needs engine RPM to decide when to change gear (automatic gearbox)
-	PTransmission.SetAllowedToChangeGear(!VehicleState.bVehicleInAir && !IsWheelSpinning());
-	float GearRatio = PTransmission.GetGearRatio(PTransmission.GetCurrentGear());
+		PTransmission.SetEngineRPM(PEngine.GetEngineRPM()); // needs engine RPM to decide when to change gear (automatic gearbox)
+		PTransmission.SetAllowedToChangeGear(!VehicleState.bVehicleInAir && !IsWheelSpinning());
+		float GearRatio = PTransmission.GetGearRatio(PTransmission.GetCurrentGear());
 
-	PTransmission.Simulate(DeltaTime);
+		PTransmission.Simulate(DeltaTime);
 
-	float TransmissionTorque = PTransmission.GetTransmissionTorque(PEngine.GetEngineTorque());
+		float TransmissionTorque = PTransmission.GetTransmissionTorque(PEngine.GetEngineTorque());
 
-	// apply drive torque to wheels
-	for (int WheelIdx = 0; WheelIdx < Wheels.Num(); WheelIdx++)
-	{
-		auto& PWheel = PVehicle->Wheels[WheelIdx];
-		if (PWheel.Setup().EngineEnabled)
+		// apply drive torque to wheels
+		for (int WheelIdx = 0; WheelIdx < Wheels.Num(); WheelIdx++)
 		{
-			check(PVehicle->NumDrivenWheels > 0);
-
-			if (PDifferential.Setup().DifferentialType == EDifferentialType::AllWheelDrive)
+			auto& PWheel = PVehicle->Wheels[WheelIdx];
+			if (PWheel.Setup().EngineEnabled)
 			{
-				float SplitTorque = 1.0f;
+				check(PVehicle->NumDrivenWheels > 0);
 
-				if (PWheel.Setup().AxleType == FSimpleWheelConfig::EAxleType::Front)
+				if (PDifferential.Setup().DifferentialType == EDifferentialType::AllWheelDrive)
 				{
-					SplitTorque = (1.0f - PDifferential.Setup().FrontRearSplit);
+					float SplitTorque = 1.0f;
+
+					if (PWheel.Setup().AxleType == FSimpleWheelConfig::EAxleType::Front)
+					{
+						SplitTorque = (1.0f - PDifferential.Setup().FrontRearSplit);
+					}
+					else
+					{
+						SplitTorque = PDifferential.Setup().FrontRearSplit;
+					}
+
+					PWheel.SetDriveTorque(MToCm(TransmissionTorque * SplitTorque) / (float)PVehicle->NumDrivenWheels);
 				}
 				else
 				{
-					SplitTorque = PDifferential.Setup().FrontRearSplit;
+					PWheel.SetDriveTorque(MToCm(TransmissionTorque) / (float)PVehicle->NumDrivenWheels);
 				}
-
-				PWheel.SetDriveTorque(MToCm(TransmissionTorque * SplitTorque) / (float)PVehicle->NumDrivenWheels);
-			}
-			else
-			{
-				PWheel.SetDriveTorque(MToCm(TransmissionTorque) / (float)PVehicle->NumDrivenWheels);
 			}
 		}
 	}
