ChaosVehicles: Wheels aware of velocity of surface they are sitting on

diff --git a/Engine/Plugins/Experimental/ChaosVehiclesPlugin/Source/ChaosVehicles/Private/ChaosWheeledVehicleMovementComponent.cpp b/Engine/Plugins/Experimental/ChaosVehiclesPlugin/Source/ChaosVehicles/Private/ChaosWheeledVehicleMovementComponent.cpp
--- a/Engine/Plugins/Experimental/ChaosVehiclesPlugin/Source/ChaosVehicles/Private/ChaosWheeledVehicleMovementComponent.cpp
+++ b/Engine/Plugins/Experimental/ChaosVehiclesPlugin/Source/ChaosVehicles/Private/ChaosWheeledVehicleMovementComponent.cpp
@@ -97,6 +97,22 @@ void FWheelState::CaptureState(int WheelIdx, const FVector& WheelOffset, const C
 	LocalWheelVelocity[WheelIdx] = WorldTransform.InverseTransformVector(WorldWheelVelocity[WheelIdx]);
 }
 
+void FWheelState::CaptureState(int WheelIdx, const FVector& WheelOffset, const Chaos::FRigidBodyHandle_Internal* VehicleHandle, const FVector& ContactPoint, const Chaos::FRigidBodyHandle_Internal* SurfaceHandle)
+{
+	check(VehicleHandle);
+
+	FVector SurfaceVelocity = FVector::ZeroVector;
+	if (SurfaceHandle)
+	{
+		SurfaceVelocity = GetVelocityAtPoint(SurfaceHandle, ContactPoint);
+	}
+
+	const FTransform WorldTransform(VehicleHandle->R(), VehicleHandle->X());
+	WheelWorldLocation[WheelIdx] = WorldTransform.TransformPosition(WheelOffset);
+	WorldWheelVelocity[WheelIdx] = GetVelocityAtPoint(VehicleHandle, WheelWorldLocation[WheelIdx]) - SurfaceVelocity;
+	LocalWheelVelocity[WheelIdx] = WorldTransform.InverseTransformVector(WorldWheelVelocity[WheelIdx]);
+}
+
 FVector FWheelState::GetVelocityAtPoint(const Chaos::FRigidBodyHandle_Internal* Rigid, const FVector& InPoint)
 {
 	if (Rigid)
@@ -150,7 +166,28 @@ void UChaosWheeledVehicleSimulation::UpdateSimulation(float DeltaTime, const FCh
 
 		for (int WheelIdx = 0; WheelIdx < Wheels.Num(); WheelIdx++)
 		{
-			WheelState.CaptureState(WheelIdx, PVehicle->Suspension[WheelIdx].GetLocalRestingPosition(), Handle);
+			bool bCaptured = false;
+
+			const FHitResult& HitResult = Wheels[WheelIdx]->HitResult;
+			if (HitResult.Component.IsValid() && HitResult.Component->GetBodyInstance())
+			{
+				const FPhysicsActorHandle& SurfaceHandle = HitResult.Component->GetBodyInstance()->ActorHandle;
+				// we are being called from the physics thread
+				if (Chaos::FRigidBodyHandle_Internal* SurfaceBody = SurfaceHandle->GetPhysicsThreadAPI())
+				{
+					if (SurfaceBody->CanTreatAsKinematic())
+					{
+						FVector Point = HitResult.ImpactPoint;
+						WheelState.CaptureState(WheelIdx, PVehicle->Suspension[WheelIdx].GetLocalRestingPosition(), Handle, Point, SurfaceBody);
+						bCaptured = true;
+					}
+				}
+			}
+
+			if (!bCaptured)
+			{
+				WheelState.CaptureState(WheelIdx, PVehicle->Suspension[WheelIdx].GetLocalRestingPosition(), Handle);
+			}
 		}
 
 		///////////////////////////////////////////////////////////////////////
diff --git a/Engine/Plugins/Experimental/ChaosVehiclesPlugin/Source/ChaosVehicles/Public/ChaosWheeledVehicleMovementComponent.h b/Engine/Plugins/Experimental/ChaosVehiclesPlugin/Source/ChaosVehicles/Public/ChaosWheeledVehicleMovementComponent.h
--- a/Engine/Plugins/Experimental/ChaosVehiclesPlugin/Source/ChaosVehicles/Public/ChaosWheeledVehicleMovementComponent.h
+++ b/Engine/Plugins/Experimental/ChaosVehiclesPlugin/Source/ChaosVehicles/Public/ChaosWheeledVehicleMovementComponent.h
@@ -517,7 +517,8 @@ struct CHAOSVEHICLES_API FWheelState
 	/** Commonly used Wheel state - evaluated once used wherever required for that frame */
 	void CaptureState(int WheelIdx, const FVector& WheelOffset, const FBodyInstance* TargetInstance);
 	void CaptureState(int WheelIdx, const FVector& WheelOffset, const Chaos::FRigidBodyHandle_Internal* Handle);
-	FVector GetVelocityAtPoint(const Chaos::FRigidBodyHandle_Internal* Rigid, const FVector& InPoint);
+	void CaptureState(int WheelIdx, const FVector& WheelOffset, const Chaos::FRigidBodyHandle_Internal* VehicleHandle, const FVector& ContactPoint, const Chaos::FRigidBodyHandle_Internal* SurfaceHandle);
+	static FVector GetVelocityAtPoint(const Chaos::FRigidBodyHandle_Internal* Rigid, const FVector& InPoint);
 
 	TArray<FVector> WheelWorldLocation;	/** Current Location Of Wheels In World Coordinates */
 	TArray<FVector> WorldWheelVelocity; /** Current velocity at wheel location In World Coordinates - combined linear and angular */
