ChaosVehicles: Fix Suspension Constraint Setup, was broken after async changes

diff --git a/Engine/Plugins/Experimental/ChaosVehiclesPlugin/Source/ChaosVehicles/Private/ChaosWheeledVehicleMovementComponent.cpp b/Engine/Plugins/Experimental/ChaosVehiclesPlugin/Source/ChaosVehicles/Private/ChaosWheeledVehicleMovementComponent.cpp
--- a/Engine/Plugins/Experimental/ChaosVehiclesPlugin/Source/ChaosVehicles/Private/ChaosWheeledVehicleMovementComponent.cpp
+++ b/Engine/Plugins/Experimental/ChaosVehiclesPlugin/Source/ChaosVehicles/Private/ChaosWheeledVehicleMovementComponent.cpp
@@ -905,6 +905,12 @@ void UChaosWheeledVehicleSimulation::FillOutputState(FChaosVehicleAsyncOutput& O
 
 }
 
+void UChaosWheeledVehicleSimulation::UpdateConstraintHandles(TArray<FPhysicsConstraintHandle>& ConstraintHandlesIn)
+{
+	UChaosVehicleSimulation::UpdateConstraintHandles(ConstraintHandlesIn);
+	ConstraintHandles = ConstraintHandlesIn;
+}
+
 /**
  * UChaosWheeledVehicleMovementComponent
  */
@@ -984,51 +990,44 @@ void UChaosWheeledVehicleMovementComponent::FixupSkeletalMesh()
 							}
 						}
 
-						if (!GWheeledVehicleDebugParams.DisableSuspensionConstraint)
-						{
-							FBodyInstance* TargetInstance = UpdatedPrimitive->GetBodyInstance();
+					}
 
-							FPhysicsCommand::ExecuteWrite(TargetInstance->ActorHandle, [&](const FPhysicsActorHandle& Chassis)
-								{
+					if (!GWheeledVehicleDebugParams.DisableSuspensionConstraint)
+					{
+						FBodyInstance* TargetInstance = UpdatedPrimitive->GetBodyInstance();
+
+						FPhysicsCommand::ExecuteWrite(TargetInstance->ActorHandle, [&](const FPhysicsActorHandle& Chassis)
+							{
 #if WITH_CHAOS
-									const FVector LocalWheel = GetWheelRestingPosition(WheelSetup);
-									FPhysicsConstraintHandle ConstraintHandle = FPhysicsInterface::CreateSuspension(Chassis, LocalWheel);
+								const FVector LocalWheel = GetWheelRestingPosition(WheelSetup);
+								FPhysicsConstraintHandle ConstraintHandle = FPhysicsInterface::CreateSuspension(Chassis, LocalWheel);
 
-									if (ConstraintHandle.IsValid())
+								if (ConstraintHandle.IsValid())
+								{
+									UChaosVehicleWheel* Wheel = Wheels[WheelIdx];
+									check(Wheel);
+									ConstraintHandles.Add(ConstraintHandle);
+									if (Chaos::FSuspensionConstraint* Constraint = static_cast<Chaos::FSuspensionConstraint*>(ConstraintHandle.Constraint))
 									{
-										UChaosVehicleWheel* Wheel = Wheels[WheelIdx];
-										check(Wheel);
-										if (Chaos::FSuspensionConstraint* Constraint = static_cast<Chaos::FSuspensionConstraint*>(ConstraintHandle.Constraint))
-										{
-											Constraint->SetHardstopStiffness(1.0f);
-											Constraint->SetSpringStiffness(Wheel->SpringRate * 0.25f);
-											Constraint->SetSpringPreload(Wheel->SpringPreload);
-											Constraint->SetSpringDamping(Wheel->SuspensionDampingRatio * 5.0f);
-											Constraint->SetMinLength(-Wheel->SuspensionMaxRaise);
-											Constraint->SetMaxLength(Wheel->SuspensionMaxDrop);
-											Constraint->SetAxis(-Wheel->SuspensionAxis);
-										}
-
-										ConstraintHandles.Add(ConstraintHandle);
-										if (Chaos::FSuspensionConstraint* Constraint = static_cast<Chaos::FSuspensionConstraint*>(ConstraintHandle.Constraint))
-										{
-											Constraint->SetHardstopStiffness(1.0f);
-											Constraint->SetSpringStiffness(Wheel->SpringRate * 0.25f);
-											Constraint->SetSpringPreload(Wheel->SpringPreload);
-											Constraint->SetSpringDamping(Wheel->SuspensionDampingRatio * 5.0f);
-											Constraint->SetMinLength(-Wheel->SuspensionMaxRaise);
-											Constraint->SetMaxLength(Wheel->SuspensionMaxDrop);
-											Constraint->SetAxis(-Wheel->SuspensionAxis);
-										}
+										Constraint->SetHardstopStiffness(1.0f);
+										Constraint->SetSpringStiffness(Chaos::MToCm(Wheel->SpringRate) * 0.25f);
+										Constraint->SetSpringPreload(Chaos::MToCm(Wheel->SpringPreload));
+										Constraint->SetSpringDamping(Wheel->SuspensionDampingRatio * 5.0f);
+										Constraint->SetMinLength(-Wheel->SuspensionMaxRaise);
+										Constraint->SetMaxLength(Wheel->SuspensionMaxDrop);
+										Constraint->SetAxis(-Wheel->SuspensionAxis);
 									}
+								}
 #endif // WITH_CHAOS
-								});
-						}
+							});
 					}
+
 				}
 			}
 		}
 
+		VehicleSimulationPT->UpdateConstraintHandles(ConstraintHandles); // TODO: think of a better way to communicate this data
+
 		Mesh->KinematicBonesUpdateType = EKinematicBonesUpdateToPhysics::SkipSimulatingBones;
 
 	}
diff --git a/Engine/Plugins/Experimental/ChaosVehiclesPlugin/Source/ChaosVehicles/Public/ChaosVehicleMovementComponent.h b/Engine/Plugins/Experimental/ChaosVehiclesPlugin/Source/ChaosVehicles/Public/ChaosVehicleMovementComponent.h
--- a/Engine/Plugins/Experimental/ChaosVehiclesPlugin/Source/ChaosVehicles/Public/ChaosVehicleMovementComponent.h
+++ b/Engine/Plugins/Experimental/ChaosVehiclesPlugin/Source/ChaosVehicles/Public/ChaosVehicleMovementComponent.h
@@ -682,6 +682,8 @@ public:
 		PVehicle = MoveTemp(PVehicleIn);
 	}
 
+	virtual void UpdateConstraintHandles(TArray<FPhysicsConstraintHandle>& ConstraintHandlesIn) {}
+
 	virtual void TickVehicle(UWorld* WorldIn, float DeltaTime, const FChaosVehicleDefaultAsyncInput& InputData, FChaosVehicleAsyncOutput& OutputData, Chaos::FRigidBodyHandle_Internal* Handle);
 
 	/** Advance the vehicle simulation */
diff --git a/Engine/Plugins/Experimental/ChaosVehiclesPlugin/Source/ChaosVehicles/Public/ChaosWheeledVehicleMovementComponent.h b/Engine/Plugins/Experimental/ChaosVehiclesPlugin/Source/ChaosVehicles/Public/ChaosWheeledVehicleMovementComponent.h
--- a/Engine/Plugins/Experimental/ChaosVehiclesPlugin/Source/ChaosVehicles/Public/ChaosWheeledVehicleMovementComponent.h
+++ b/Engine/Plugins/Experimental/ChaosVehiclesPlugin/Source/ChaosVehicles/Public/ChaosWheeledVehicleMovementComponent.h
@@ -532,9 +532,8 @@ class UChaosWheeledVehicleSimulation : public UChaosVehicleSimulation
 {
 public:
 
-	UChaosWheeledVehicleSimulation(TArray<class UChaosVehicleWheel*>& WheelsIn, TArray<FPhysicsConstraintHandle>& ConstraintHandlesIn)
+	UChaosWheeledVehicleSimulation(TArray<class UChaosVehicleWheel*>& WheelsIn)
 		: Wheels(WheelsIn)
-		, ConstraintHandles(ConstraintHandlesIn)
 	{
 
 	}
@@ -550,6 +549,8 @@ public:
 		WheelState.Init(PVehicle->Wheels.Num());
 	}
 
+	virtual void UpdateConstraintHandles(TArray<FPhysicsConstraintHandle>& ConstraintHandlesIn) override;
+
 	virtual void TickVehicle(UWorld* WorldIn, float DeltaTime, const FChaosVehicleDefaultAsyncInput& InputData, FChaosVehicleAsyncOutput& OutputData, Chaos::FRigidBodyHandle_Internal* Handle) override;
 
 	/** Advance the vehicle simulation */
@@ -590,7 +591,7 @@ public:
 
 	FWheelState WheelState;	/** Cached state that holds wheel data for this frame */
 
-	TArray<FPhysicsConstraintHandle>& ConstraintHandles;
+	TArray<FPhysicsConstraintHandle> ConstraintHandles;
 
 };
 
@@ -747,7 +748,7 @@ protected:
 	virtual TUniquePtr<Chaos::FSimpleWheeledVehicle> CreatePhysicsVehicle() override
 	{
 		// Make the Vehicle Simulation class that will be updated from the physics thread async callback
-		VehicleSimulationPT = MakeUnique<UChaosWheeledVehicleSimulation>(Wheels, ConstraintHandles);
+		VehicleSimulationPT = MakeUnique<UChaosWheeledVehicleSimulation>(Wheels);
 
 		return UChaosVehicleMovementComponent::CreatePhysicsVehicle();
 	}
