Chaos : Vehicle

-Early exit UChaosWheeledVehicleSimulation::UpdateSimulation when game thread data is not initialized yet.

#okforgithum public, internal

diff --git a/Engine/Plugins/Experimental/ChaosVehiclesPlugin/Source/ChaosVehicles/Private/ChaosWheeledVehicleMovementComponent.cpp b/Engine/Plugins/Experimental/ChaosVehiclesPlugin/Source/ChaosVehicles/Private/ChaosWheeledVehicleMovementComponent.cpp
--- a/Engine/Plugins/Experimental/ChaosVehiclesPlugin/Source/ChaosVehicles/Private/ChaosWheeledVehicleMovementComponent.cpp
+++ b/Engine/Plugins/Experimental/ChaosVehiclesPlugin/Source/ChaosVehicles/Private/ChaosWheeledVehicleMovementComponent.cpp
@@ -152,6 +152,24 @@ void UChaosWheeledVehicleSimulation::UpdateSimulation(float DeltaTime, const FCh
 	// Inherit common vehicle simulation stages ApplyAerodynamics, ApplyTorqueControl, etc
 	UChaosVehicleSimulation::UpdateSimulation(DeltaTime, InputData, Handle);
 
+#if WITH_CHAOS
+	// @todo(Chaos) Warning : This function is not thread safe. 
+	// the body instance is game thread data being used on the physics
+	// thread. This HACK just exits if any part of the body instance is 
+	// not initialized yet. 
+	for (int WheelIdx = 0; WheelIdx < Wheels.Num(); WheelIdx++)
+	{
+		const FHitResult& HitResult = Wheels[WheelIdx]->HitResult;
+		if (HitResult.Component.IsValid() && HitResult.Component->GetBodyInstance())
+		{
+			if (ensure(!HitResult.Component->GetBodyInstance()->GetPhysicsActorHandle()))
+			{
+				return;
+			}
+		}
+	}
+#endif
+
 	if (CanSimulate() && Handle)
 	{
 		// sanity check that everything is setup ok
@@ -171,15 +189,17 @@ void UChaosWheeledVehicleSimulation::UpdateSimulation(float DeltaTime, const FCh
 			const FHitResult& HitResult = Wheels[WheelIdx]->HitResult;
 			if (HitResult.Component.IsValid() && HitResult.Component->GetBodyInstance())
 			{
-				const FPhysicsActorHandle& SurfaceHandle = HitResult.Component->GetBodyInstance()->ActorHandle;
-				// we are being called from the physics thread
-				if (Chaos::FRigidBodyHandle_Internal* SurfaceBody = SurfaceHandle->GetPhysicsThreadAPI())
+				if (const FPhysicsActorHandle& SurfaceHandle = HitResult.Component->GetBodyInstance()->GetPhysicsActorHandle())
 				{
-					if (SurfaceBody->CanTreatAsKinematic())
+					// we are being called from the physics thread
+					if (Chaos::FRigidBodyHandle_Internal* SurfaceBody = SurfaceHandle->GetPhysicsThreadAPI())
 					{
-						FVector Point = HitResult.ImpactPoint;
-						WheelState.CaptureState(WheelIdx, PVehicle->Suspension[WheelIdx].GetLocalRestingPosition(), Handle, Point, SurfaceBody);
-						bCaptured = true;
+						if (SurfaceBody->CanTreatAsKinematic())
+						{
+							FVector Point = HitResult.ImpactPoint;
+							WheelState.CaptureState(WheelIdx, PVehicle->Suspension[WheelIdx].GetLocalRestingPosition(), Handle, Point, SurfaceBody);
+							bCaptured = true;
+						}
 					}
 				}
 			}
