Fix CIS

diff --git a/Engine/Plugins/Experimental/ChaosVehiclesPlugin/Source/ChaosVehicles/Public/ChaosVehicleMovementComponent.h b/Engine/Plugins/Experimental/ChaosVehiclesPlugin/Source/ChaosVehicles/Public/ChaosVehicleMovementComponent.h
--- a/Engine/Plugins/Experimental/ChaosVehiclesPlugin/Source/ChaosVehicles/Public/ChaosVehicleMovementComponent.h
+++ b/Engine/Plugins/Experimental/ChaosVehiclesPlugin/Source/ChaosVehicles/Public/ChaosVehicleMovementComponent.h
@@ -458,7 +458,7 @@ struct CHAOSVEHICLES_API FVehicleInputRateConfig
 	UPROPERTY(EditAnywhere, Category = VehicleInputRate)
 	FRuntimeFloatCurve UserCurve;
 
-	FVehicleInputRateConfig() : RiseRate(5.0f), FallRate(5.0f) { }
+	FVehicleInputRateConfig() : RiseRate(5.0f), FallRate(5.0f), InputCurveFunction(EInputFunctionType::LinearFunction) { }
 
 	/** Change an output value using max rise and fall rates */
 	float InterpInputValue( float DeltaTime, float CurrentValue, float NewValue ) const
