ChaosVehicles: Small change to allow controls to be taken from one vehicle and fed to another. For cab and trailer, where the trailer wheels brake in conjunction with the cab.
ChaosVehicles: slight bug with the way steering curve is translated from UE vehicle layer to core vehicle layer.
Editgrate from UE5EarlyAccess

diff --git a/Engine/Plugins/Experimental/ChaosVehiclesPlugin/Source/ChaosVehicles/Private/ChaosVehicleMovementComponent.cpp b/Engine/Plugins/Experimental/ChaosVehiclesPlugin/Source/ChaosVehicles/Private/ChaosVehicleMovementComponent.cpp
--- a/Engine/Plugins/Experimental/ChaosVehiclesPlugin/Source/ChaosVehicles/Private/ChaosVehicleMovementComponent.cpp
+++ b/Engine/Plugins/Experimental/ChaosVehiclesPlugin/Source/ChaosVehicles/Private/ChaosVehicleMovementComponent.cpp
@@ -624,6 +624,7 @@ UChaosVehicleMovementComponent::UChaosVehicleMovementComponent(const FObjectInit
 
 	AngErrorAccumulator = 0.0f;
 
+	bRequiresControllerForInputs = true;
 	IdleBrakeInput = 0.0f;
 	StopThreshold = 10.0f; 
 	WrongDirectionThreshold = 100.f;
@@ -1109,9 +1110,11 @@ void UChaosVehicleMovementComponent::UpdateState(float DeltaTime)
 	AController* Controller = GetController();
 	VehicleState.CaptureState(GetBodyInstance(), GetGravityZ(), DeltaTime);
 
+	bool bProcessLocally = bRequiresControllerForInputs?(Controller && Controller->IsLocalController()):true;
+
 	// IsLocallyControlled will fail if the owner is unpossessed (i.e. Controller == nullptr);
 	// Should we remove input instead of relying on replicated state in that case?
-	if (Controller && Controller->IsLocalController() && PVehicleOutput)
+	if (bProcessLocally && PVehicleOutput)
 	{
 		if (bReverseAsBrake)
 		{
diff --git a/Engine/Plugins/Experimental/ChaosVehiclesPlugin/Source/ChaosVehicles/Public/ChaosVehicleMovementComponent.h b/Engine/Plugins/Experimental/ChaosVehiclesPlugin/Source/ChaosVehicles/Public/ChaosVehicleMovementComponent.h
--- a/Engine/Plugins/Experimental/ChaosVehiclesPlugin/Source/ChaosVehicles/Public/ChaosVehicleMovementComponent.h
+++ b/Engine/Plugins/Experimental/ChaosVehiclesPlugin/Source/ChaosVehicles/Public/ChaosVehicleMovementComponent.h
@@ -958,6 +958,19 @@ public:
 	UFUNCTION(BlueprintCallable, Category = "Game|Components|ChaosVehicleMovement")
 	float GetForwardSpeedMPH() const;
 
+	/** Get the user input for the vehicle throttle - can use this to feed control to a connected trailer */
+	UFUNCTION(BlueprintCallable, Category = "Game|Components|ChaosVehicleMovement")
+	float GetThrottleInput() { return RawThrottleInput; }
+
+	/** Get the user input for the vehicle brake - can use this to feed control to a connected trailer */
+	UFUNCTION(BlueprintCallable, Category = "Game|Components|ChaosVehicleMovement")
+	float GetBrakeInput() { return RawBrakeInput; }
+
+	/** Get the user input for the vehicle steering - can use this to feed control to a connected trailer */
+	UFUNCTION(BlueprintCallable, Category = "Game|Components|ChaosVehicleMovement")
+	float GetSteeringInput() { return RawSteeringInput; }
+
+
 	/** Is the vehicle in park mode */
 	UFUNCTION(BlueprintCallable, Category = "Game|Components|ChaosVehicleMovement")
 	bool IsParked() const;
@@ -1077,6 +1090,10 @@ protected:
 	UPROPERTY(Transient)
 	float HandbrakeInput;
 
+	// Bypass the need for a controller in order for the controls to be processed.
+	UPROPERTY(EditAnywhere, Category=VehicleInput)
+	bool bRequiresControllerForInputs;
+
 	// How much to press the brake when the player has release throttle
 	UPROPERTY(EditAnywhere, Category=VehicleInput)
 	float IdleBrakeInput;
diff --git a/Engine/Plugins/Experimental/ChaosVehiclesPlugin/Source/ChaosVehicles/Public/ChaosWheeledVehicleMovementComponent.h b/Engine/Plugins/Experimental/ChaosVehiclesPlugin/Source/ChaosVehicles/Public/ChaosWheeledVehicleMovementComponent.h
--- a/Engine/Plugins/Experimental/ChaosVehiclesPlugin/Source/ChaosVehicles/Public/ChaosWheeledVehicleMovementComponent.h
+++ b/Engine/Plugins/Experimental/ChaosVehiclesPlugin/Source/ChaosVehicles/Public/ChaosWheeledVehicleMovementComponent.h
@@ -464,7 +464,7 @@ private:
 		float MaxX = this->SteeringCurve.GetRichCurveConst()->GetLastKey().Time;
 		PSteeringConfig.SpeedVsSteeringCurve.Empty();
 		float NumSamples = 20;
-		for (float X = MinValue; X <= MaxX; X += ((MaxX- MinValue) / NumSamples))
+		for (float X = 0; X <= MaxX; X += (MaxX / NumSamples))
 		{
 			float Y = this->SteeringCurve.GetRichCurveConst()->Eval(X) / MaxValue;
 			PSteeringConfig.SpeedVsSteeringCurve.Add(FVector2D(X, Y));
