ChaosVehicles:
* Use correct PhysicsParallelFor for physics thread in vehicle simulation callback.
* Use chaios debug drawing in simulation callback
* Fix issue with custom steering curve only steering right.

diff --git a/Engine/Plugins/Experimental/ChaosVehiclesPlugin/Source/ChaosVehicles/Private/AsyncCallback.cpp b/Engine/Plugins/Experimental/ChaosVehiclesPlugin/Source/ChaosVehicles/Private/AsyncCallback.cpp
--- a/Engine/Plugins/Experimental/ChaosVehiclesPlugin/Source/ChaosVehicles/Private/AsyncCallback.cpp
+++ b/Engine/Plugins/Experimental/ChaosVehiclesPlugin/Source/ChaosVehicles/Private/AsyncCallback.cpp
@@ -69,7 +69,7 @@ void FChaosVehicleManagerAsyncCallback::OnPreSimulate_Internal()
 	};
 
 	bool ForceSingleThread = !GVehicleDebugParams.EnableMultithreading;
-	ParallelFor(OutputVehiclesBatch.Num(), LambdaParallelUpdate, ForceSingleThread);
+	PhysicsParallelFor(OutputVehiclesBatch.Num(), LambdaParallelUpdate, ForceSingleThread);
 
 
 }
diff --git a/Engine/Plugins/Experimental/ChaosVehiclesPlugin/Source/ChaosVehicles/Private/ChaosWheeledVehicleMovementComponent.cpp b/Engine/Plugins/Experimental/ChaosVehiclesPlugin/Source/ChaosVehicles/Private/ChaosWheeledVehicleMovementComponent.cpp
--- a/Engine/Plugins/Experimental/ChaosVehiclesPlugin/Source/ChaosVehicles/Private/ChaosWheeledVehicleMovementComponent.cpp
+++ b/Engine/Plugins/Experimental/ChaosVehiclesPlugin/Source/ChaosVehicles/Private/ChaosWheeledVehicleMovementComponent.cpp
@@ -16,6 +16,8 @@
 #include "SteeringUtility.h"
 #include "Chaos/ChaosEngineInterface.h"
 #include "Chaos/PBDSuspensionConstraintData.h"
+#include "Chaos/DebugDrawQueue.h"
+
 
 #if !(UE_BUILD_SHIPPING || UE_BUILD_TEST)
 #include "CanvasItem.h"
@@ -83,6 +85,7 @@ void FWheelState::CaptureState(int WheelIdx, const FVector& WheelOffset, const F
 {
 	check(TargetInstance);
 	const FTransform WorldTransform = TargetInstance->GetUnrealWorldTransform();
+	WheelLocalLocation[WheelIdx] = WheelOffset;
 	WheelWorldLocation[WheelIdx] = WorldTransform.TransformPosition(WheelOffset);
 	WorldWheelVelocity[WheelIdx] = TargetInstance->GetUnrealWorldVelocityAtPoint(WheelWorldLocation[WheelIdx]);
 	LocalWheelVelocity[WheelIdx] = WorldTransform.InverseTransformVector(WorldWheelVelocity[WheelIdx]);
@@ -92,6 +95,7 @@ void FWheelState::CaptureState(int WheelIdx, const FVector& WheelOffset, const C
 {
 	check(Handle);
 	const FTransform WorldTransform(Handle->R(), Handle->X());
+	WheelLocalLocation[WheelIdx] = WheelOffset;
 	WheelWorldLocation[WheelIdx] = WorldTransform.TransformPosition(WheelOffset);
 	WorldWheelVelocity[WheelIdx] = GetVelocityAtPoint(Handle, WheelWorldLocation[WheelIdx]);
 	LocalWheelVelocity[WheelIdx] = WorldTransform.InverseTransformVector(WorldWheelVelocity[WheelIdx]);
@@ -108,6 +112,7 @@ void FWheelState::CaptureState(int WheelIdx, const FVector& WheelOffset, const C
 	}
 
 	const FTransform WorldTransform(VehicleHandle->R(), VehicleHandle->X());
+	WheelLocalLocation[WheelIdx] = WheelOffset;
 	WheelWorldLocation[WheelIdx] = WorldTransform.TransformPosition(WheelOffset);
 	WorldWheelVelocity[WheelIdx] = GetVelocityAtPoint(VehicleHandle, WheelWorldLocation[WheelIdx]) - SurfaceVelocity;
 	LocalWheelVelocity[WheelIdx] = WorldTransform.InverseTransformVector(WorldWheelVelocity[WheelIdx]);
@@ -307,7 +312,7 @@ void UChaosWheeledVehicleSimulation::PerformSuspensionTraces(const TArray<FSuspe
 #if !(UE_BUILD_SHIPPING || UE_BUILD_TEST)
 		if (GWheeledVehicleDebugParams.ShowBatchQueryExtents)
 		{
-			DrawDebugBox(World, QueryBox.GetCenter(), QueryBox.GetExtent(), FColor::Yellow, false, -1.0f, 0, 2.0f);
+			Chaos::FDebugDrawQueue::GetInstance().DrawDebugBox(QueryBox.GetCenter(), QueryBox.GetExtent(), FQuat::Identity, FColor::Yellow, false, -1.0f, 0, 2.0f);
 		}
 #endif
 
@@ -461,19 +466,16 @@ void UChaosWheeledVehicleSimulation::ApplyWheelFrictionForces(float DeltaTime)
 			if (GWheeledVehicleDebugParams.ShowWheelForces)
 			{
 				// show longitudinal drive force
-				{
-					DrawDebugLine(World
-						, WheelState.WheelWorldLocation[WheelIdx]
+				FDebugDrawQueue::GetInstance().DrawDebugLine(
+						  WheelState.WheelWorldLocation[WheelIdx]
 						, WheelState.WheelWorldLocation[WheelIdx] + FrictionForceVector * 0.001f
 						, FColor::Yellow, false, -1.0f, 0, 2);
 
-					DrawDebugLine(World
-						, WheelState.WheelWorldLocation[WheelIdx]
+				FDebugDrawQueue::GetInstance().DrawDebugLine(
+						  WheelState.WheelWorldLocation[WheelIdx]
 						, WheelState.WheelWorldLocation[WheelIdx] + GroundZVector * 100.f
 						, FColor::Orange, false, -1.0f, 0, 2);
 
-				}
-
 			}
 #endif
 
@@ -513,6 +515,8 @@ void UChaosWheeledVehicleSimulation::ApplySuspensionForces(float DeltaTime)
 					if (Chaos::FSuspensionConstraint* Constraint = static_cast<Chaos::FSuspensionConstraint*>(ConstraintHandle.Constraint))
 					{
 						FVec3 P = HitResult.ImpactPoint + (PWheel.Setup().WheelRadius * VehicleState.VehicleUpAxis);
+						
+						// This isn't threadsafe because of FDirtySet ProxiesData.Add(Base);
 						Constraint->SetTarget(P);
 						Constraint->SetEnabled(PWheel.InContact());
 					}
@@ -553,13 +557,13 @@ void UChaosWheeledVehicleSimulation::ApplySuspensionForces(float DeltaTime)
 #if !(UE_BUILD_SHIPPING || UE_BUILD_TEST)
 			if (GWheeledVehicleDebugParams.ShowSuspensionForces)
 			{
-				DrawDebugLine(World
-					, SusApplicationPoint
+				Chaos::FDebugDrawQueue::GetInstance().DrawDebugLine(
+					  SusApplicationPoint
 					, SusApplicationPoint + SuspensionForceVector * GVehicleDebugParams.ForceDebugScaling
 					, FColor::Blue, false, -1.0f, 0, 5);
 
-				DrawDebugLine(World
-					, SusApplicationPoint
+				Chaos::FDebugDrawQueue::GetInstance().DrawDebugLine(
+					  SusApplicationPoint
 					, SusApplicationPoint + GroundZVector * 140.f
 					, FColor::Yellow, false, -1.0f, 0, 5);
 			}
@@ -818,12 +822,13 @@ void UChaosWheeledVehicleSimulation::DrawDebug3D()
 			FVector Start = WorldLocation + WorldDirection * (PWheel.GetEffectiveRadius() - PSuspension.Setup().SuspensionMaxRaise);
 			FVector End = WorldLocation + WorldDirection * (PWheel.GetEffectiveRadius() + PSuspension.Setup().SuspensionMaxDrop);
 
-			DrawDebugLine(World, Start + VehicleRightAxis, End + VehicleRightAxis, FColor::Orange, false, -1.f, 0, 3.f);
+			Chaos::FDebugDrawQueue::GetInstance().DrawDebugLine(
+				Start + VehicleRightAxis, End + VehicleRightAxis, FColor::Orange, false, -1.f, 0, 3.f);
 
 			FVector Start2 = WorldLocation - WorldDirection * PSuspension.Setup().SuspensionMaxRaise;
 			FVector End2 = WorldLocation + WorldDirection * PSuspension.Setup().SuspensionMaxDrop;
 
-			DrawDebugLine(World, Start2 + VehicleRightAxis, End2 + VehicleRightAxis, FColor::Yellow, false, -1.f, 0, 3.f);
+			FDebugDrawQueue::GetInstance().DrawDebugLine(Start2 + VehicleRightAxis, End2 + VehicleRightAxis, FColor::Yellow, false, -1.f, 0, 3.f);
 		}
 	}
 
@@ -832,7 +837,7 @@ void UChaosWheeledVehicleSimulation::DrawDebug3D()
 		for (int WheelIdx = 0; WheelIdx < Wheels.Num(); WheelIdx++)
 		{
 			FHitResult& Hit = Wheels[WheelIdx]->HitResult;
-			DrawDebugLine(World, Hit.ImpactPoint, Hit.ImpactPoint + Hit.Normal * 20.0f, FColor::Yellow, true, 1.0f, 0, 1.0f);
+			FDebugDrawQueue::GetInstance().DrawDebugLine(Hit.ImpactPoint, Hit.ImpactPoint + Hit.Normal * 20.0f, FColor::Yellow, true, 1.0f, 0, 1.0f);
 		}
 	}
 
@@ -852,10 +857,10 @@ void UChaosWheeledVehicleSimulation::DrawDebug3D()
 			}
 
 			FColor UseColor = PVehicle->Wheels[WheelIdx].InContact() ? FColor::Green : FColor::Red;
-			DrawDebugDirectionalArrow(World, TraceStart + VehicleRightAxis, TraceEnd + VehicleRightAxis, 10.f, UseColor, false, -1.f, 0, 2.f);
+			Chaos::FDebugDrawQueue::GetInstance().DrawDebugDirectionalArrow(TraceStart + VehicleRightAxis, TraceEnd + VehicleRightAxis, 10.f, UseColor, false, -1.f, 0, 2.f);
 
-			DrawDebugLine(World, TraceStart, TraceStart + VehicleRightAxis, FColor::White, false, -1.f, 0, 1.f);
-			DrawDebugLine(World, TraceEnd, TraceEnd + VehicleRightAxis, FColor::White, false, -1.f, 0, 1.f);
+			Chaos::FDebugDrawQueue::GetInstance().DrawDebugLine(TraceStart, TraceStart + VehicleRightAxis, FColor::White, false, -1.f, 0, 1.f);
+			Chaos::FDebugDrawQueue::GetInstance().DrawDebugLine(TraceEnd, TraceEnd + VehicleRightAxis, FColor::White, false, -1.f, 0, 1.f);
 		}
 	}
 #endif
diff --git a/Engine/Plugins/Experimental/ChaosVehiclesPlugin/Source/ChaosVehicles/Public/ChaosVehicleMovementComponent.h b/Engine/Plugins/Experimental/ChaosVehiclesPlugin/Source/ChaosVehicles/Public/ChaosVehicleMovementComponent.h
--- a/Engine/Plugins/Experimental/ChaosVehiclesPlugin/Source/ChaosVehicles/Public/ChaosVehicleMovementComponent.h
+++ b/Engine/Plugins/Experimental/ChaosVehiclesPlugin/Source/ChaosVehicles/Public/ChaosVehicleMovementComponent.h
@@ -487,8 +487,8 @@ struct CHAOSVEHICLES_API FVehicleInputRateConfig
 		{
 			if (UserCurve.GetRichCurveConst() && !UserCurve.GetRichCurveConst()->IsEmpty())
 			{
-				float Output = UserCurve.GetRichCurveConst()->Eval(InputValue);
-				return FMath::Clamp(Output, 0.0f, 1.0f); // Ensure output is normalized
+				float Output = FMath::Clamp(UserCurve.GetRichCurveConst()->Eval(FMath::Abs(InputValue)), 0.0f, 1.0f);
+				return (InputValue < 0.f)? -Output : Output;
 			}
 			else
 			{
diff --git a/Engine/Plugins/Experimental/ChaosVehiclesPlugin/Source/ChaosVehicles/Public/ChaosWheeledVehicleMovementComponent.h b/Engine/Plugins/Experimental/ChaosVehiclesPlugin/Source/ChaosVehicles/Public/ChaosWheeledVehicleMovementComponent.h
--- a/Engine/Plugins/Experimental/ChaosVehiclesPlugin/Source/ChaosVehicles/Public/ChaosWheeledVehicleMovementComponent.h
+++ b/Engine/Plugins/Experimental/ChaosVehiclesPlugin/Source/ChaosVehicles/Public/ChaosWheeledVehicleMovementComponent.h
@@ -508,6 +508,7 @@ struct CHAOSVEHICLES_API FWheelState
 {
 	void Init(int NumWheels)
 	{
+		WheelLocalLocation.Init(FVector::ZeroVector, NumWheels);
 		WheelWorldLocation.Init(FVector::ZeroVector, NumWheels);
 		WorldWheelVelocity.Init(FVector::ZeroVector, NumWheels);
 		LocalWheelVelocity.Init(FVector::ZeroVector, NumWheels);
@@ -520,6 +521,7 @@ struct CHAOSVEHICLES_API FWheelState
 	void CaptureState(int WheelIdx, const FVector& WheelOffset, const Chaos::FRigidBodyHandle_Internal* VehicleHandle, const FVector& ContactPoint, const Chaos::FRigidBodyHandle_Internal* SurfaceHandle);
 	static FVector GetVelocityAtPoint(const Chaos::FRigidBodyHandle_Internal* Rigid, const FVector& InPoint);
 
+	TArray<FVector> WheelLocalLocation;	/** Current Location Of Wheels In Local Coordinates */
 	TArray<FVector> WheelWorldLocation;	/** Current Location Of Wheels In World Coordinates */
 	TArray<FVector> WorldWheelVelocity; /** Current velocity at wheel location In World Coordinates - combined linear and angular */
 	TArray<FVector> LocalWheelVelocity; /** Local velocity of Wheel */
