Unshelved from pending changelist '15569632':

Blueprint changable properties.

diff --git a/Engine/Plugins/Experimental/ChaosVehiclesPlugin/Source/ChaosVehicles/Private/ChaosVehicleMovementComponent.cpp b/Engine/Plugins/Experimental/ChaosVehiclesPlugin/Source/ChaosVehicles/Private/ChaosVehicleMovementComponent.cpp
--- a/Engine/Plugins/Experimental/ChaosVehiclesPlugin/Source/ChaosVehicles/Private/ChaosVehicleMovementComponent.cpp
+++ b/Engine/Plugins/Experimental/ChaosVehiclesPlugin/Source/ChaosVehicles/Private/ChaosVehicleMovementComponent.cpp
@@ -968,18 +968,18 @@ float UChaosVehicleMovementComponent::CalcSteeringInput()
 
 void UChaosVehicleMovementComponent::CalcThrottleBrakeInput(float& ThrottleOut, float& BrakeOut)
 {
+	BrakeOut = RawBrakeInput;
+	ThrottleOut = RawThrottleInput;
+
 	if (bReverseAsBrake)
 	{
-		BrakeOut = RawBrakeInput;
-		ThrottleOut = RawThrottleInput;
-
 		if (RawThrottleInput > 0.f)
 		{
 			// car moving backwards but player wants to move forwards...
 			// if vehicle is moving backwards, then press brake
 			if (VehicleState.ForwardSpeed < -WrongDirectionThreshold)
 			{
-				BrakeOut = 1.0f; // = RawThrottleInput ?
+				BrakeOut = 1.0f;
 				ThrottleOut = 0.0f;
 			}
 
@@ -1125,13 +1125,16 @@ void UChaosVehicleMovementComponent::UpdateState(float DeltaTime)
 		}
 		else
 		{
-			if (TransmissionType == Chaos::ETransmissionType::Automatic
-				&& RawThrottleInput > KINDA_SMALL_NUMBER
-				&& PVehicleOutput->CurrentGear == 0
-				&& PVehicleOutput->TargetGear == 0)
+			if (TransmissionType == Chaos::ETransmissionType::Automatic)
 			{
-				SetTargetGear(1, true);
+				if (RawThrottleInput > KINDA_SMALL_NUMBER
+					&& GetCurrentGear() == 0
+					&& GetTargetGear() == 0)
+				{
+					SetTargetGear(1, true);
+				}
 			}
+
 		}
 
 		float ModifiedThrottle = 0.f;
diff --git a/Engine/Plugins/Experimental/ChaosVehiclesPlugin/Source/ChaosVehicles/Private/ChaosWheeledVehicleMovementComponent.cpp b/Engine/Plugins/Experimental/ChaosVehiclesPlugin/Source/ChaosVehicles/Private/ChaosWheeledVehicleMovementComponent.cpp
--- a/Engine/Plugins/Experimental/ChaosVehiclesPlugin/Source/ChaosVehicles/Private/ChaosWheeledVehicleMovementComponent.cpp
+++ b/Engine/Plugins/Experimental/ChaosVehiclesPlugin/Source/ChaosVehicles/Private/ChaosWheeledVehicleMovementComponent.cpp
@@ -205,7 +205,7 @@ void UChaosWheeledVehicleSimulation::UpdateSimulation(float DeltaTime, const FCh
 		{
 			auto& PSuspension = PVehicle->Suspension[WheelIdx];
 			auto& PWheel = PVehicle->Wheels[WheelIdx];
-			PSuspension.UpdateWorldRaycastLocation(VehicleState.VehicleWorldTransform, PWheel.Setup().WheelRadius, WheelState.Trace[WheelIdx]);
+			PSuspension.UpdateWorldRaycastLocation(VehicleState.VehicleWorldTransform, PWheel.GetEffectiveRadius(), WheelState.Trace[WheelIdx]);
 		}
 
 		if (!GWheeledVehicleDebugParams.DisableSuspensionForces && PVehicle->bSuspensionEnabled)
@@ -514,7 +514,7 @@ void UChaosWheeledVehicleSimulation::ApplySuspensionForces(float DeltaTime)
 				{
 					if (Chaos::FSuspensionConstraint* Constraint = static_cast<Chaos::FSuspensionConstraint*>(ConstraintHandle.Constraint))
 					{
-						FVec3 P = HitResult.ImpactPoint + (PWheel.Setup().WheelRadius * VehicleState.VehicleUpAxis);
+						FVec3 P = HitResult.ImpactPoint + (PWheel.GetEffectiveRadius() * VehicleState.VehicleUpAxis);
 						
 						// This isn't threadsafe because of FDirtySet ProxiesData.Add(Base);
 						Constraint->SetTarget(P);
@@ -530,9 +530,9 @@ void UChaosWheeledVehicleSimulation::ApplySuspensionForces(float DeltaTime)
 		{
 			NewDesiredLength = HitResult.Distance;
 
-			SuspensionMovePosition = -FVector::DotProduct(WheelState.WheelWorldLocation[WheelIdx] - HitResult.ImpactPoint, VehicleState.VehicleUpAxis) + PWheel.Setup().WheelRadius;
+			SuspensionMovePosition = -FVector::DotProduct(WheelState.WheelWorldLocation[WheelIdx] - HitResult.ImpactPoint, VehicleState.VehicleUpAxis) + PWheel.GetEffectiveRadius();
 
-			PSuspension.SetSuspensionLength(NewDesiredLength, PWheel.Setup().WheelRadius);
+			PSuspension.SetSuspensionLength(NewDesiredLength, PWheel.GetEffectiveRadius());
 			PSuspension.SetLocalVelocity(WheelState.LocalWheelVelocity[WheelIdx]);
 			PSuspension.Simulate(DeltaTime);
 
@@ -572,7 +572,7 @@ void UChaosWheeledVehicleSimulation::ApplySuspensionForces(float DeltaTime)
 		}
 		else
 		{
-			PSuspension.SetSuspensionLength(PSuspension.GetTraceLength(PWheel.Setup().WheelRadius), PWheel.Setup().WheelRadius);
+			PSuspension.SetSuspensionLength(PSuspension.GetTraceLength(PWheel.GetEffectiveRadius()), PWheel.Setup().WheelRadius);
 		}
 
 	}
@@ -611,7 +611,7 @@ void UChaosWheeledVehicleSimulation::ProcessSteering(const FControlInputs& Contr
 		auto& PWheel = PVehicle->Wheels[WheelIdx]; // Physics Wheel
 		FHitResult& HitResult = Wheels[WheelIdx]->HitResult;
 
-		if (PWheel.Setup().SteeringEnabled)
+		if (PWheel.SteeringEnabled)
 		{
 			float SpeedScale = 1.0f;
 
@@ -625,12 +625,12 @@ void UChaosWheeledVehicleSimulation::ProcessSteering(const FControlInputs& Contr
 
 			if (FMath::Abs(GWheeledVehicleDebugParams.SteeringOverride) > 0.01f)
 			{
-				SteeringAngle = PWheel.Setup().MaxSteeringAngle * GWheeledVehicleDebugParams.SteeringOverride;
+				SteeringAngle = PWheel.MaxSteeringAngle * GWheeledVehicleDebugParams.SteeringOverride;
 			}
 			else
 			{
 				float WheelSide = PVehicle->GetSuspension(WheelIdx).GetLocalRestingPosition().Y;
-				SteeringAngle = PSteering.GetSteeringAngle(SteeringAngle, PWheel.Setup().MaxSteeringAngle, WheelSide);
+				SteeringAngle = PSteering.GetSteeringAngle(SteeringAngle, PWheel.MaxSteeringAngle, WheelSide);
 			}
 
 			PWheel.SetSteeringAngle(SteeringAngle);
@@ -690,14 +690,14 @@ void UChaosWheeledVehicleSimulation::ApplyInput(const FControlInputs& ControlInp
 		auto& PWheel = PVehicle->Wheels[WheelIdx];
 
 		float EngineBrakingForce = 0.0f;
-		if ((ModifiedInputs.ThrottleInput < SMALL_NUMBER) && FMath::Abs(VehicleState.ForwardSpeed) > SMALL_NUMBER && PWheel.Setup().EngineEnabled)
+		if ((ModifiedInputs.ThrottleInput < SMALL_NUMBER) && FMath::Abs(VehicleState.ForwardSpeed) > SMALL_NUMBER && PWheel.EngineEnabled)
 		{
 			EngineBrakingForce = EngineBraking;
 		}
 
-		if (PWheel.Setup().BrakeEnabled)
+		if (PWheel.BrakeEnabled)
 		{
-			float BrakeForce = PWheel.Setup().MaxBrakeTorque * ModifiedInputs.BrakeInput;
+			float BrakeForce = PWheel.MaxBrakeTorque * ModifiedInputs.BrakeInput;
 			PWheel.SetBrakeTorque(TorqueMToCm(BrakeForce + EngineBrakingForce));
 		}
 		else
@@ -705,9 +705,9 @@ void UChaosWheeledVehicleSimulation::ApplyInput(const FControlInputs& ControlInp
 			PWheel.SetBrakeTorque(TorqueMToCm(EngineBraking));
 		}
 
-		if ((ModifiedInputs.HandbrakeInput && PWheel.Setup().HandbrakeEnabled) || ModifiedInputs.ParkingEnabled)
+		if ((ModifiedInputs.HandbrakeInput && PWheel.HandbrakeEnabled) || ModifiedInputs.ParkingEnabled)
 		{
-			float HandbrakeForce = ModifiedInputs.ParkingEnabled ? PWheel.Setup().HandbrakeTorque : (ModifiedInputs.HandbrakeInput * PWheel.Setup().HandbrakeTorque);
+			float HandbrakeForce = ModifiedInputs.ParkingEnabled ? PWheel.HandbrakeTorque : (ModifiedInputs.HandbrakeInput * PWheel.HandbrakeTorque);
 			PWheel.SetBrakeTorque(TorqueMToCm(HandbrakeForce));
 		}
 	}
@@ -739,7 +739,7 @@ void UChaosWheeledVehicleSimulation::ProcessMechanicalSimulation(float DeltaTime
 		float WheelRPM = 0;
 		for (int I = 0; I < PVehicle->Wheels.Num(); I++)
 		{
-			if (PVehicle->Wheels[I].Setup().EngineEnabled)
+			if (PVehicle->Wheels[I].EngineEnabled)
 			{
 				WheelRPM = FMath::Abs(PVehicle->Wheels[I].GetWheelRPM());
 			}
@@ -760,7 +760,7 @@ void UChaosWheeledVehicleSimulation::ProcessMechanicalSimulation(float DeltaTime
 		for (int WheelIdx = 0; WheelIdx < Wheels.Num(); WheelIdx++)
 		{
 			auto& PWheel = PVehicle->Wheels[WheelIdx];
-			if (PWheel.Setup().EngineEnabled)
+			if (PWheel.EngineEnabled)
 			{
 				check(PVehicle->NumDrivenWheels > 0);
 
@@ -770,11 +770,11 @@ void UChaosWheeledVehicleSimulation::ProcessMechanicalSimulation(float DeltaTime
 
 					if (PWheel.Setup().AxleType == FSimpleWheelConfig::EAxleType::Front)
 					{
-						SplitTorque = (1.0f - PDifferential.Setup().FrontRearSplit);
+						SplitTorque = (1.0f - PDifferential.FrontRearSplit);
 					}
 					else
 					{
-						SplitTorque = PDifferential.Setup().FrontRearSplit;
+						SplitTorque = PDifferential.FrontRearSplit;
 					}
 
 					PWheel.SetDriveTorque(TorqueMToCm(TransmissionTorque * SplitTorque) / (float)PVehicle->NumDrivenWheels);
@@ -2105,8 +2105,9 @@ void UChaosWheeledVehicleMovementComponent::SetSnapshot(const FWheeledSnaphotDat
 	UChaosVehicleMovementComponent::SetBaseSnapshot(SnapshotIn);
 
 	FBodyInstance* TargetInstance = GetBodyInstance();
-
-	FPhysicsCommand::ExecuteWrite(TargetInstance->ActorHandle, [&](const FPhysicsActorHandle& Chassis)
+	if (TargetInstance)
+	{
+		FPhysicsCommand::ExecuteWrite(TargetInstance->ActorHandle, [&](const FPhysicsActorHandle& Chassis)
 		{
 			const FWheeledSnaphotData* WheelSnapshotData = static_cast<const FWheeledSnaphotData*>(&SnapshotIn);
 
@@ -2132,9 +2133,246 @@ void UChaosWheeledVehicleMovementComponent::SetSnapshot(const FWheeledSnaphotDat
 				}
 			}
 		});
+	}
+}
+
+//////////////////////////////////////////////////////////////////////////
+
+void UChaosWheeledVehicleMovementComponent::SetMaxEngineTorque(float Torque)
+{
+	if (FBodyInstance* TargetInstance = GetBodyInstance())
+	{
+		FPhysicsCommand::ExecuteWrite(TargetInstance->ActorHandle, [&](const FPhysicsActorHandle& Chassis)
+		{
+			if (VehicleSimulationPT && VehicleSimulationPT->PVehicle)
+			{
+				VehicleSimulationPT->PVehicle->GetEngine().SetMaxTorque(Torque);
+			}
+		});
+	}
+}
+
+void UChaosWheeledVehicleMovementComponent::SetDragCoefficient(float DragCoeff)
+{
+	if (FBodyInstance* TargetInstance = GetBodyInstance())
+	{
+		FPhysicsCommand::ExecuteWrite(TargetInstance->ActorHandle, [&](const FPhysicsActorHandle& Chassis)
+		{
+			if (VehicleSimulationPT && VehicleSimulationPT->PVehicle)
+			{
+				VehicleSimulationPT->PVehicle->GetAerodynamics().SetDragCoefficient(DragCoeff);
+			}
+		});
+	}
+}
+
+void UChaosWheeledVehicleMovementComponent::SetDownforceCoefficient(float DownforceCoeff)
+{
+	if (FBodyInstance* TargetInstance = GetBodyInstance())
+	{
+		FPhysicsCommand::ExecuteWrite(TargetInstance->ActorHandle, [&](const FPhysicsActorHandle& Chassis)
+		{
+			if (VehicleSimulationPT && VehicleSimulationPT->PVehicle)
+			{
+				VehicleSimulationPT->PVehicle->GetAerodynamics().SetDownforceCoefficient(DownforceCoeff);
+			}
+		});
+	}
+}
+
+void UChaosWheeledVehicleMovementComponent::SetDifferentialFrontRearSplit(float FrontRearSplit)
+{
+	if (FBodyInstance* TargetInstance = GetBodyInstance())
+	{
+		FPhysicsCommand::ExecuteWrite(TargetInstance->ActorHandle, [&](const FPhysicsActorHandle& Chassis)
+		{
+			if (VehicleSimulationPT && VehicleSimulationPT->PVehicle)
+			{
+				VehicleSimulationPT->PVehicle->GetDifferential().FrontRearSplit = FrontRearSplit;
+			}
+		});
+	}
+}
+
+
+void UChaosWheeledVehicleMovementComponent::SetTractionControlEnabled(int WheelIndex, bool Enabled)
+{
+	if (FBodyInstance* TargetInstance = GetBodyInstance())
+	{
+		FPhysicsCommand::ExecuteWrite(TargetInstance->ActorHandle, [&](const FPhysicsActorHandle& Chassis)
+		{
+			if (VehicleSimulationPT && VehicleSimulationPT->PVehicle && WheelIndex < VehicleSimulationPT->PVehicle->Wheels.Num())
+			{
+				Chaos::FSimpleWheelSim& VehicleWheel = VehicleSimulationPT->PVehicle->Wheels[WheelIndex];
 
+				VehicleWheel.TractionControlEnabled = Enabled;
+			}
+		});
+	}
 }
 
+void UChaosWheeledVehicleMovementComponent::SetABSEnabled(int WheelIndex, bool Enabled)
+{
+	if (FBodyInstance* TargetInstance = GetBodyInstance())
+	{
+		FPhysicsCommand::ExecuteWrite(TargetInstance->ActorHandle, [&](const FPhysicsActorHandle& Chassis)
+		{
+			if (VehicleSimulationPT && VehicleSimulationPT->PVehicle && WheelIndex < VehicleSimulationPT->PVehicle->Wheels.Num())
+			{
+				Chaos::FSimpleWheelSim& VehicleWheel = VehicleSimulationPT->PVehicle->Wheels[WheelIndex];
+
+				VehicleWheel.ABSEnabled = Enabled;
+			}
+		});
+	}
+}
+
+void UChaosWheeledVehicleMovementComponent::SetAffectedByBrake(int WheelIndex, bool Enabled)
+{
+	if (FBodyInstance* TargetInstance = GetBodyInstance())
+	{
+		FPhysicsCommand::ExecuteWrite(TargetInstance->ActorHandle, [&](const FPhysicsActorHandle& Chassis)
+		{
+			if (VehicleSimulationPT && VehicleSimulationPT->PVehicle && WheelIndex < VehicleSimulationPT->PVehicle->Wheels.Num())
+			{
+				Chaos::FSimpleWheelSim& VehicleWheel = VehicleSimulationPT->PVehicle->Wheels[WheelIndex];
+
+				VehicleWheel.BrakeEnabled = Enabled;
+			}
+		});
+	}
+}
+
+void UChaosWheeledVehicleMovementComponent::SetAffectedByHandbrake(int WheelIndex, bool Enabled)
+{
+	if (FBodyInstance* TargetInstance = GetBodyInstance())
+	{
+		FPhysicsCommand::ExecuteWrite(TargetInstance->ActorHandle, [&](const FPhysicsActorHandle& Chassis)
+		{
+			if (VehicleSimulationPT && VehicleSimulationPT->PVehicle && WheelIndex < VehicleSimulationPT->PVehicle->Wheels.Num())
+			{
+				Chaos::FSimpleWheelSim& VehicleWheel = VehicleSimulationPT->PVehicle->Wheels[WheelIndex];
+
+				VehicleWheel.HandbrakeEnabled = Enabled; // this is affecting all vehicles - needs to be per instance 
+			}
+		});
+	}
+}
+
+void UChaosWheeledVehicleMovementComponent::SetAffectedBySteering(int WheelIndex, bool Enabled)
+{
+	if (FBodyInstance* TargetInstance = GetBodyInstance())
+	{
+		FPhysicsCommand::ExecuteWrite(TargetInstance->ActorHandle, [&](const FPhysicsActorHandle& Chassis)
+		{
+			if (VehicleSimulationPT && VehicleSimulationPT->PVehicle && WheelIndex < VehicleSimulationPT->PVehicle->Wheels.Num())
+			{
+				Chaos::FSimpleWheelSim& VehicleWheel = VehicleSimulationPT->PVehicle->Wheels[WheelIndex];
+
+				VehicleWheel.SteeringEnabled = Enabled;
+			}
+		});
+	}
+}
+
+void UChaosWheeledVehicleMovementComponent::SetAffectedByEngine(int WheelIndex, bool Enabled)
+{
+	if (FBodyInstance* TargetInstance = GetBodyInstance())
+	{
+		FPhysicsCommand::ExecuteWrite(TargetInstance->ActorHandle, [&](const FPhysicsActorHandle& Chassis)
+		{
+			if (VehicleSimulationPT && VehicleSimulationPT->PVehicle && WheelIndex < VehicleSimulationPT->PVehicle->Wheels.Num())
+			{
+				Chaos::FSimpleWheelSim& VehicleWheel = VehicleSimulationPT->PVehicle->Wheels[WheelIndex];
+
+				VehicleWheel.EngineEnabled = Enabled;
+			}
+		});
+	}
+
+}
+
+void UChaosWheeledVehicleMovementComponent::SetWheelRadius(int WheelIndex, float Radius)
+{
+	if (FBodyInstance* TargetInstance = GetBodyInstance())
+	{
+		FPhysicsCommand::ExecuteWrite(TargetInstance->ActorHandle, [&](const FPhysicsActorHandle& Chassis)
+		{
+			if (VehicleSimulationPT && VehicleSimulationPT->PVehicle && WheelIndex < VehicleSimulationPT->PVehicle->Wheels.Num())
+			{
+				Chaos::FSimpleWheelSim& VehicleWheel = VehicleSimulationPT->PVehicle->Wheels[WheelIndex];
+
+				VehicleWheel.SetWheelRadius(Radius);
+			}
+		});
+	}
+}
+
+void UChaosWheeledVehicleMovementComponent::SetWheelFrictionMultiplier(int WheelIndex, float Friction)
+{
+	if (FBodyInstance* TargetInstance = GetBodyInstance())
+	{
+		FPhysicsCommand::ExecuteWrite(TargetInstance->ActorHandle, [&](const FPhysicsActorHandle& Chassis)
+			{
+				if (VehicleSimulationPT && VehicleSimulationPT->PVehicle && WheelIndex < VehicleSimulationPT->PVehicle->Wheels.Num())
+				{
+					Chaos::FSimpleWheelSim& VehicleWheel = VehicleSimulationPT->PVehicle->Wheels[WheelIndex];
+
+					VehicleWheel.FrictionMultiplier = Friction;
+				}
+			});
+	}
+}
+
+void UChaosWheeledVehicleMovementComponent::SetWheelMaxBrakeTorque(int WheelIndex, float Torque)
+{
+	if (FBodyInstance* TargetInstance = GetBodyInstance())
+	{
+		FPhysicsCommand::ExecuteWrite(TargetInstance->ActorHandle, [&](const FPhysicsActorHandle& Chassis)
+		{
+			if (VehicleSimulationPT && VehicleSimulationPT->PVehicle && WheelIndex < VehicleSimulationPT->PVehicle->Wheels.Num())
+			{
+				Chaos::FSimpleWheelSim& VehicleWheel = VehicleSimulationPT->PVehicle->Wheels[WheelIndex];
+
+				VehicleWheel.MaxBrakeTorque = Torque;
+			}
+		});
+	}
+}
+
+void UChaosWheeledVehicleMovementComponent::SetWheelHandbrakeTorque(int WheelIndex, float Torque)
+{
+	if (FBodyInstance* TargetInstance = GetBodyInstance())
+	{
+		FPhysicsCommand::ExecuteWrite(TargetInstance->ActorHandle, [&](const FPhysicsActorHandle& Chassis)
+		{
+			if (VehicleSimulationPT && VehicleSimulationPT->PVehicle && WheelIndex < VehicleSimulationPT->PVehicle->Wheels.Num())
+			{
+				Chaos::FSimpleWheelSim& VehicleWheel = VehicleSimulationPT->PVehicle->Wheels[WheelIndex];
+
+				VehicleWheel.HandbrakeTorque = Torque;
+			}
+		});
+	}
+}
+
+void UChaosWheeledVehicleMovementComponent::SetWheelMaxSteerAngle(int WheelIndex, float AngleDegrees)
+{
+	if (FBodyInstance* TargetInstance = GetBodyInstance())
+	{
+		FPhysicsCommand::ExecuteWrite(TargetInstance->ActorHandle, [&](const FPhysicsActorHandle& Chassis)
+		{
+			if (VehicleSimulationPT && VehicleSimulationPT->PVehicle && WheelIndex < VehicleSimulationPT->PVehicle->Wheels.Num())
+			{
+				Chaos::FSimpleWheelSim& VehicleWheel = VehicleSimulationPT->PVehicle->Wheels[WheelIndex];
+
+				VehicleWheel.MaxSteeringAngle = AngleDegrees;
+			}
+		});
+	}
+}
+
+
 FChaosWheelSetup::FChaosWheelSetup()
 	: WheelClass(UChaosVehicleWheel::StaticClass())
 //	, SteeringBoneName(NAME_None)
diff --git a/Engine/Plugins/Experimental/ChaosVehiclesPlugin/Source/ChaosVehicles/Public/ChaosVehicleWheel.h b/Engine/Plugins/Experimental/ChaosVehiclesPlugin/Source/ChaosVehicles/Public/ChaosVehicleWheel.h
--- a/Engine/Plugins/Experimental/ChaosVehiclesPlugin/Source/ChaosVehicles/Public/ChaosVehicleWheel.h
+++ b/Engine/Plugins/Experimental/ChaosVehiclesPlugin/Source/ChaosVehicles/Public/ChaosVehicleWheel.h
@@ -364,7 +364,6 @@ class UChaosWheeledVehicleMovementComponent;
 			PWheelConfig.MaxSteeringAngle = this->MaxSteerAngle;
 			PWheelConfig.MaxBrakeTorque = this->MaxBrakeTorque;
 			PWheelConfig.HandbrakeTorque = this->MaxHandBrakeTorque;
-
 			PWheelConfig.SteeringEnabled = this->bAffectedBySteering;
 			PWheelConfig.BrakeEnabled = this->bAffectedByBrake;
 			PWheelConfig.HandbrakeEnabled = this->bAffectedByHandbrake;
@@ -375,7 +374,7 @@ class UChaosWheeledVehicleMovementComponent;
 			PWheelConfig.LongitudinalFrictionMultiplier = this->LongitudinalFrictionForceMultiplier;
 			PWheelConfig.LateralFrictionMultiplier = this->LateralFrictionForceMultiplier;
 			PWheelConfig.FrictionMultiplier = this->FrictionForceMultiplier;
-			PWheelConfig.CorneringStiffness = this->CorneringStiffness;
+			PWheelConfig.CorneringStiffness = this->CorneringStiffness * 10000.0f;
 			PWheelConfig.NewSimulationPath = this->bNewWheelSimulation;
 			PWheelConfig.SideSlipModifier = this->SideSlipModifier;
 			PWheelConfig.SlipThreshold = this->SlipThreshold;
diff --git a/Engine/Plugins/Experimental/ChaosVehiclesPlugin/Source/ChaosVehicles/Public/ChaosWheeledVehicleMovementComponent.h b/Engine/Plugins/Experimental/ChaosVehiclesPlugin/Source/ChaosVehicles/Public/ChaosWheeledVehicleMovementComponent.h
--- a/Engine/Plugins/Experimental/ChaosVehiclesPlugin/Source/ChaosVehicles/Public/ChaosWheeledVehicleMovementComponent.h
+++ b/Engine/Plugins/Experimental/ChaosVehiclesPlugin/Source/ChaosVehicles/Public/ChaosWheeledVehicleMovementComponent.h
@@ -156,26 +156,6 @@ struct FVehicleDifferentialConfig
 	UPROPERTY(EditAnywhere, Category = Setup, meta = (ClampMin = "0.0", UIMin = "0.0", ClampMax = "1.0", UIMax = "1.0"))
 	float FrontRearSplit;
 
-	///** Ratio of torque split between front-left and front-right (>0.5 means more to front-left, <0.5 means more to front-right, works only with 4W and LimitedSlip_FrontDrive) */
-	//UPROPERTY(EditAnywhere, Category = Setup, meta = (ClampMin = "0.0", UIMin = "0.0", ClampMax = "1.0", UIMax = "1.0"))
-	//float FrontLeftRightSplit;
-
-	///** Ratio of torque split between rear-left and rear-right (>0.5 means more to rear-left, <0.5 means more to rear-right, works only with 4W and LimitedSlip_RearDrive) */
-	//UPROPERTY(EditAnywhere, Category = Setup, meta = (ClampMin = "0.0", UIMin = "0.0", ClampMax = "1.0", UIMax = "1.0"))
-	//float RearLeftRightSplit;
-	//
-	///** Maximum allowed ratio of average front wheel rotation speed and rear wheel rotation speeds (range: 1..inf, works only with LimitedSlip_4W) */
-	//UPROPERTY(EditAnywhere, Category = Setup, meta = (ClampMin = "1.0", UIMin = "1.0"))
-	//float CentreBias;
-
-	///** Maximum allowed ratio of front-left and front-right wheel rotation speeds (range: 1..inf, works only with LimitedSlip_4W, LimitedSlip_FrontDrive) */
-	//UPROPERTY(EditAnywhere, Category = Setup, meta = (ClampMin = "1.0", UIMin = "1.0"))
-	//float FrontBias;
-
-	///** Maximum allowed ratio of rear-left and rear-right wheel rotation speeds (range: 1..inf, works only with LimitedSlip_4W, LimitedSlip_FrontDrive) */
-	//UPROPERTY(EditAnywhere, Category = Setup, meta = (ClampMin = "1.0", UIMin = "1.0"))
-	//float RearBias;
-
 	const Chaos::FSimpleDifferentialConfig& GetPhysicsDifferentialConfig()
 	{
 		FillDifferentialSetup();
@@ -760,6 +740,53 @@ class CHAOSVEHICLES_API UChaosWheeledVehicleMovementComponent : public UChaosVeh
 	virtual void SetSnapshot(const FWheeledSnaphotData& SnapshotIn);
 
 
+	//////////////////////////////////////////////////////////////////////////
+	// change handling via blueprint at runtime
+	UFUNCTION(BlueprintCallable, Category = "Game|Components|ChaosWheeledVehicleMovement")
+	void SetMaxEngineTorque(float Torque);
+
+	UFUNCTION(BlueprintCallable, Category = "Game|Components|ChaosWheeledVehicleMovement")
+	void SetDragCoefficient(float DragCoeff);
+
+	UFUNCTION(BlueprintCallable, Category = "Game|Components|ChaosWheeledVehicleMovement")
+	void SetDownforceCoefficient(float DownforceCoeff);
+
+	UFUNCTION(BlueprintCallable, Category = "Game|Components|ChaosWheeledVehicleMovement")
+	void SetDifferentialFrontRearSplit(float FrontRearSlpit);
+
+	UFUNCTION(BlueprintCallable, Category = "Game|Components|ChaosWheeledVehicleMovement")
+	void SetTractionControlEnabled(int WheelIndex, bool Enabled);
+
+	UFUNCTION(BlueprintCallable, Category = "Game|Components|ChaosWheeledVehicleMovement")
+	void SetABSEnabled(int WheelIndex, bool Enabled);
+
+	UFUNCTION(BlueprintCallable, Category = "Game|Components|ChaosWheeledVehicleMovement")
+	void SetAffectedByBrake(int WheelIndex, bool Enabled);
+
+	UFUNCTION(BlueprintCallable, Category = "Game|Components|ChaosWheeledVehicleMovement")
+	void SetAffectedByHandbrake(int WheelIndex, bool Enabled);
+
+	UFUNCTION(BlueprintCallable, Category = "Game|Components|ChaosWheeledVehicleMovement")
+	void SetAffectedBySteering(int WheelIndex, bool Enabled);
+
+	UFUNCTION(BlueprintCallable, Category = "Game|Components|ChaosWheeledVehicleMovement")
+	void SetAffectedByEngine(int WheelIndex, bool Enabled);
+
+	UFUNCTION(BlueprintCallable, Category = "Game|Components|ChaosWheeledVehicleMovement")
+	void SetWheelRadius(int WheelIndex, float Radius);
+
+	UFUNCTION(BlueprintCallable, Category = "Game|Components|ChaosWheeledVehicleMovement")
+	void SetWheelFrictionMultiplier(int WheelIndex, float Friction);
+
+	UFUNCTION(BlueprintCallable, Category = "Game|Components|ChaosWheeledVehicleMovement")
+	void SetWheelMaxBrakeTorque(int WheelIndex, float Torque);
+
+	UFUNCTION(BlueprintCallable, Category = "Game|Components|ChaosWheeledVehicleMovement")
+	void SetWheelHandbrakeTorque(int WheelIndex, float Torque);
+
+	UFUNCTION(BlueprintCallable, Category = "Game|Components|ChaosWheeledVehicleMovement")
+	void SetWheelMaxSteerAngle(int WheelIndex, float AngleDegrees);
+
 protected:
 
 	//////////////////////////////////////////////////////////////////////////
diff --git a/Engine/Source/Runtime/Experimental/ChaosVehicles/ChaosVehiclesCore/Private/AerodynamicsSystem.cpp b/Engine/Source/Runtime/Experimental/ChaosVehicles/ChaosVehiclesCore/Private/AerodynamicsSystem.cpp
--- a/Engine/Source/Runtime/Experimental/ChaosVehicles/ChaosVehiclesCore/Private/AerodynamicsSystem.cpp
+++ b/Engine/Source/Runtime/Experimental/ChaosVehicles/ChaosVehiclesCore/Private/AerodynamicsSystem.cpp
@@ -10,11 +10,13 @@ namespace Chaos
 {
 
 	FSimpleAerodynamicsSim::FSimpleAerodynamicsSim(const FSimpleAerodynamicsConfig* SetupIn) : TVehicleSystem< FSimpleAerodynamicsConfig>(SetupIn)
+		, DownforceCoefficient(SetupIn->DownforceCoefficient)
+		, DragCoefficient(SetupIn->DragCoefficient)
 		, DensityOfMedium(RealWorldConsts::AirDensity())
 	{
 		// pre-calculate static values
-		EffectiveDragConstant = 0.5f * Setup().AreaMetresSquared * Setup().DragCoefficient;
-		EffectiveLiftConstant = 0.5f * Setup().AreaMetresSquared * Setup().DownforceCoefficient;
+		EffectiveDragConstant = 0.5f * Setup().AreaMetresSquared * DragCoefficient;
+		EffectiveLiftConstant = 0.5f * Setup().AreaMetresSquared * DownforceCoefficient;
 	}
 
 	FVector FSimpleAerodynamicsSim::GetCombinedForces(float VelocityIn)
diff --git a/Engine/Source/Runtime/Experimental/ChaosVehicles/ChaosVehiclesCore/Private/EngineSystem.cpp b/Engine/Source/Runtime/Experimental/ChaosVehicles/ChaosVehiclesCore/Private/EngineSystem.cpp
--- a/Engine/Source/Runtime/Experimental/ChaosVehicles/ChaosVehiclesCore/Private/EngineSystem.cpp
+++ b/Engine/Source/Runtime/Experimental/ChaosVehicles/ChaosVehiclesCore/Private/EngineSystem.cpp
@@ -11,6 +11,7 @@ namespace Chaos
 
 
 	FSimpleEngineSim::FSimpleEngineSim(const FSimpleEngineConfig* StaticDataIn) : TVehicleSystem<FSimpleEngineConfig>(StaticDataIn)
+		, MaxTorque(Setup().MaxTorque)
 		, ThrottlePosition(0.f)
 		, TargetSpeed(0.f)
 		, CurrentRPM(0.f)
@@ -37,7 +38,7 @@ namespace Chaos
 			RPM = FMath::Clamp(RPM, (float)Setup().EngineIdleRPM, (float)Setup().MaxRPM);
 		}
 
-		return Setup().TorqueCurve.GetValue(RPM, Setup().MaxRPM, Setup().MaxTorque);
+		return Setup().TorqueCurve.GetValue(RPM, Setup().MaxRPM, MaxTorque);
 	}
 
 	void FSimpleEngineSim::Simulate(float DeltaTime)
diff --git a/Engine/Source/Runtime/Experimental/ChaosVehicles/ChaosVehiclesCore/Private/WheelSystem.cpp b/Engine/Source/Runtime/Experimental/ChaosVehicles/ChaosVehiclesCore/Private/WheelSystem.cpp
--- a/Engine/Source/Runtime/Experimental/ChaosVehicles/ChaosVehiclesCore/Private/WheelSystem.cpp
+++ b/Engine/Source/Runtime/Experimental/ChaosVehicles/ChaosVehiclesCore/Private/WheelSystem.cpp
@@ -9,6 +9,17 @@ PRAGMA_DISABLE_OPTIMIZATION
 namespace Chaos
 {
 	FSimpleWheelSim::FSimpleWheelSim(const FSimpleWheelConfig* SetupIn) : TVehicleSystem<FSimpleWheelConfig>(SetupIn)
+		, BrakeEnabled(SetupIn->BrakeEnabled)
+		, HandbrakeEnabled(SetupIn->HandbrakeEnabled)
+		, SteeringEnabled(SetupIn->SteeringEnabled)
+		, EngineEnabled(SetupIn->EngineEnabled)
+		, TractionControlEnabled(SetupIn->TractionControlEnabled)
+		, ABSEnabled(SetupIn->ABSEnabled)
+		, FrictionMultiplier(SetupIn->FrictionMultiplier)
+		, CorneringStiffness(SetupIn->CorneringStiffness)
+		, MaxSteeringAngle(SetupIn->MaxSteeringAngle)
+		, MaxBrakeTorque(SetupIn->MaxBrakeTorque)
+		, HandbrakeTorque(SetupIn->HandbrakeTorque)
 		, Re(SetupIn->WheelRadius)
 		, Omega(0.f)
 		, Sx(0.f)
@@ -103,7 +114,7 @@ namespace Chaos
 					BrakeFactor = FMath::Clamp(LongitudinalAdhesiveLimit / FMath::Abs(FinalLongitudinalForce), 0.6f, 1.0f);
 				}
 
-				if ((Braking && Setup().ABSEnabled) || (!Braking && Setup().TractionControlEnabled))
+				if ((Braking && ABSEnabled) || (!Braking && TractionControlEnabled))
 				{
 					Spin = 0.0f;
 
@@ -203,7 +214,7 @@ namespace Chaos
 		AppliedLinearBrakeForce = BrakeTorque / Re;
 
 		// Longitudinal multiplier now affecting both brake and steering equally
-		float AvailableGrip = ForceIntoSurface * SurfaceFriction * Setup().FrictionMultiplier;
+		float AvailableGrip = ForceIntoSurface * SurfaceFriction * FrictionMultiplier;
 
 		float FinalLongitudinalForce = 0.f;
 		float FinalLateralForce = 0.f;
@@ -216,9 +227,9 @@ namespace Chaos
 		if (ForceIntoSurface > SMALL_NUMBER)
 		{
 			// ABS limiting brake force to match force from the grip avialable
-			if (Setup().ABSEnabled && Braking && FMath::Abs(AppliedLinearBrakeForce) > AvailableGrip)
+			if (ABSEnabled && Braking && FMath::Abs(AppliedLinearBrakeForce) > AvailableGrip)
 			{
-				if ((Braking && Setup().ABSEnabled) || (!Braking && Setup().TractionControlEnabled))
+				if ((Braking && ABSEnabled) || (!Braking && TractionControlEnabled))
 				{
 					float Sign = (AppliedLinearBrakeForce > 0.0f) ? 1.0f : -1.0f;
 					AppliedLinearBrakeForce = AvailableGrip * TractionControlAndABSScaling * Sign;
@@ -226,7 +237,7 @@ namespace Chaos
 			}
 
 			// Traction control limiting drive force to match force from grip available
-			if (Setup().TractionControlEnabled && !Braking && FMath::Abs(AppliedLinearDriveForce) > AvailableGrip)
+			if (TractionControlEnabled && !Braking && FMath::Abs(AppliedLinearDriveForce) > AvailableGrip)
 			{
 				float Sign = (AppliedLinearDriveForce > 0.0f) ? 1.0f : -1.0f;
 				AppliedLinearDriveForce = AvailableGrip * TractionControlAndABSScaling * Sign;
@@ -258,7 +269,6 @@ namespace Chaos
 			float ForceRequiredToBringToStop = -(MassPerWheel * K * GroundVelocityVector.Y) / DeltaTime;
 
 			// use slip angle to generate a sideways force
-			float CorneringStiffness = Setup().CorneringStiffness * 10000.0f;
 
 			// Leveling off of the graph
 			float AngleLimit = FMath::DegreesToRadians(8.0f);
diff --git a/Engine/Source/Runtime/Experimental/ChaosVehicles/ChaosVehiclesCore/Public/AerodynamicsSystem.h b/Engine/Source/Runtime/Experimental/ChaosVehicles/ChaosVehiclesCore/Public/AerodynamicsSystem.h
--- a/Engine/Source/Runtime/Experimental/ChaosVehicles/ChaosVehiclesCore/Public/AerodynamicsSystem.h
+++ b/Engine/Source/Runtime/Experimental/ChaosVehicles/ChaosVehiclesCore/Public/AerodynamicsSystem.h
@@ -47,6 +47,16 @@ namespace Chaos
 			DensityOfMedium = DensityIn;
 		}
 
+		void SetDragCoefficient(float InCoeffient)
+		{
+			DragCoefficient = InCoeffient;
+		}
+
+		void SetDownforceCoefficient(float InCoeffient)
+		{
+			DownforceCoefficient = InCoeffient;
+		}
+
 		/** get the drag force generated at the given velocity */
 		float GetDragForceFromVelocity(float VelocityIn)
 		{
@@ -63,6 +73,8 @@ namespace Chaos
 		FVector GetCombinedForces(float VelocityIn);
 
 	protected:
+		float DownforceCoefficient;
+		float DragCoefficient;
 		float DensityOfMedium;
 		float EffectiveDragConstant;
 		float EffectiveLiftConstant;
diff --git a/Engine/Source/Runtime/Experimental/ChaosVehicles/ChaosVehiclesCore/Public/EngineSystem.h b/Engine/Source/Runtime/Experimental/ChaosVehicles/ChaosVehiclesCore/Public/EngineSystem.h
--- a/Engine/Source/Runtime/Experimental/ChaosVehicles/ChaosVehiclesCore/Public/EngineSystem.h
+++ b/Engine/Source/Runtime/Experimental/ChaosVehicles/ChaosVehiclesCore/Public/EngineSystem.h
@@ -90,6 +90,11 @@ namespace Chaos
 			}
 		}
 
+		void SetMaxTorque(float InTorque)
+		{
+			MaxTorque = InTorque;
+		}
+
 		float GetEngineTorque()
 		{
 			if (EngineStarted)
@@ -132,7 +137,7 @@ namespace Chaos
 		void Simulate(float DeltaTime);
 
 	protected:
-
+		float MaxTorque;		// [N.m] The peak torque Y value in the normalized torque graph
 		float ThrottlePosition; // [0..1 Normalized position]
 		float TargetSpeed;		// target RPM
 		float CurrentRPM;		// current RPM
diff --git a/Engine/Source/Runtime/Experimental/ChaosVehicles/ChaosVehiclesCore/Public/TransmissionSystem.h b/Engine/Source/Runtime/Experimental/ChaosVehicles/ChaosVehiclesCore/Public/TransmissionSystem.h
--- a/Engine/Source/Runtime/Experimental/ChaosVehicles/ChaosVehiclesCore/Public/TransmissionSystem.h
+++ b/Engine/Source/Runtime/Experimental/ChaosVehicles/ChaosVehiclesCore/Public/TransmissionSystem.h
@@ -48,8 +48,10 @@ namespace Chaos
 	{
 	public:
 		FSimpleDifferentialSim(const FSimpleDifferentialConfig* SetupIn) 
-			: TVehicleSystem<FSimpleDifferentialConfig>(SetupIn) {}
+			: TVehicleSystem<FSimpleDifferentialConfig>(SetupIn)
+			, FrontRearSplit(Setup().FrontRearSplit) {}
 
+		float FrontRearSplit;
 	};
 
 	struct CHAOSVEHICLESCORE_API FSimpleTransmissionConfig
diff --git a/Engine/Source/Runtime/Experimental/ChaosVehicles/ChaosVehiclesCore/Public/WheelSystem.h b/Engine/Source/Runtime/Experimental/ChaosVehicles/ChaosVehiclesCore/Public/WheelSystem.h
--- a/Engine/Source/Runtime/Experimental/ChaosVehicles/ChaosVehiclesCore/Public/WheelSystem.h
+++ b/Engine/Source/Runtime/Experimental/ChaosVehicles/ChaosVehiclesCore/Public/WheelSystem.h
@@ -389,6 +389,17 @@ public:
 	}
 
 public:
+	bool BrakeEnabled;			// Regular brakes are enabled for this wheel
+	bool HandbrakeEnabled;		// Handbrake is operational on this wheel
+	bool SteeringEnabled;		// Steering is operational on this wheel
+	bool EngineEnabled;			// Wheel is driven by an engine
+	bool TractionControlEnabled;// Straight Line Traction Control
+	bool ABSEnabled;			// Advanced braking system operational
+	float FrictionMultiplier;
+	float CorneringStiffness;
+	float MaxSteeringAngle;
+	float MaxBrakeTorque;
+	float HandbrakeTorque;
 
 	float Re;		// [cm] Effective Wheel Radius could change dynamically if get a flat?, tire shreds
 	float Omega;	// [radians/sec] Wheel Rotation Angular Velocity
