ChaosVehicles: a few small additions
* Debug drawing friction circle debug graphics
* Wheels spinning better
* Reverse As Brake, no longer tries to apply brake when rolling backwards and applying the accelerator. It will wheel spin now.
* Exposed current SlipAngle to WheelStatus, accessible from blueprint.
* User can now specify a custom lateral slip graph. lateral force vs slip angle for ultimate control in setup.

diff --git a/Engine/Plugins/Experimental/ChaosVehiclesPlugin/Source/ChaosVehicles/Private/ChaosVehicleMovementComponent.cpp b/Engine/Plugins/Experimental/ChaosVehiclesPlugin/Source/ChaosVehicles/Private/ChaosVehicleMovementComponent.cpp
--- a/Engine/Plugins/Experimental/ChaosVehiclesPlugin/Source/ChaosVehicles/Private/ChaosVehicleMovementComponent.cpp
+++ b/Engine/Plugins/Experimental/ChaosVehiclesPlugin/Source/ChaosVehicles/Private/ChaosVehicleMovementComponent.cpp
@@ -975,13 +975,16 @@ void UChaosVehicleMovementComponent::CalcThrottleBrakeInput(float& ThrottleOut,
 	{
 		if (RawThrottleInput > 0.f)
 		{
-			// car moving backwards but player wants to move forwards...
-			// if vehicle is moving backwards, then press brake
-			if (VehicleState.ForwardSpeed < -WrongDirectionThreshold)
-			{
-				BrakeOut = 1.0f;
-				ThrottleOut = 0.0f;
-			}
+		// Note: Removed this condition to support wheel spinning when rolling backareds with accelerator pressed, rather than braking
+		// Should this case be another checkbox option??
+		//	
+		//	// car moving backwards but player wants to move forwards...
+		//	// if vehicle is moving backwards, then press brake
+		//	if (VehicleState.ForwardSpeed < -WrongDirectionThreshold)
+		//	{
+		//		BrakeOut = 1.0f;
+		//		ThrottleOut = 0.0f;
+		//	}
 
 		}
 		else if (RawBrakeInput > 0.f)
@@ -1111,7 +1114,8 @@ void UChaosVehicleMovementComponent::UpdateState(float DeltaTime)
 		if (bReverseAsBrake)
 		{
 			//for reverse as state we want to automatically shift between reverse and first gear
-			if (FMath::Abs(GetForwardSpeed()) < WrongDirectionThreshold)	//we only shift between reverse and first if the car is slow enough.
+			// Note: Removed this condition to support wheel spinning when rolling backareds with accelerator pressed, rather than braking
+			//if (FMath::Abs(GetForwardSpeed()) < WrongDirectionThreshold)	//we only shift between reverse and first if the car is slow enough.
 			{
 				if (RawBrakeInput > KINDA_SMALL_NUMBER && GetCurrentGear() >= 0 && GetTargetGear() >= 0)
 				{
@@ -1757,6 +1761,7 @@ void UChaosVehicleMovementComponent::ParallelUpdate(float DeltaSeconds)
 					PVehicleOutput->Wheels[WheelIdx].bIsSkidding = Current.bIsSkidding;
 					PVehicleOutput->Wheels[WheelIdx].SkidMagnitude = FMath::Lerp(Current.SkidMagnitude, Next.SkidMagnitude, OutputInterpAlpha);
 					PVehicleOutput->Wheels[WheelIdx].SkidNormal = FMath::Lerp(Current.SkidNormal, Next.SkidNormal, OutputInterpAlpha);
+					PVehicleOutput->Wheels[WheelIdx].SlipAngle = FMath::Lerp(Current.SlipAngle, Next.SlipAngle, OutputInterpAlpha);
 
 					PVehicleOutput->Wheels[WheelIdx].SuspensionOffset = FMath::Lerp(Current.SuspensionOffset, Next.SuspensionOffset, OutputInterpAlpha);
 					PVehicleOutput->Wheels[WheelIdx].SpringForce = FMath::Lerp(Current.SpringForce, Next.SpringForce, OutputInterpAlpha);
@@ -1787,6 +1792,7 @@ void UChaosVehicleMovementComponent::ParallelUpdate(float DeltaSeconds)
 					PVehicleOutput->Wheels[WheelIdx].bIsSkidding = Current.bIsSkidding;
 					PVehicleOutput->Wheels[WheelIdx].SkidMagnitude = Current.SkidMagnitude;
 					PVehicleOutput->Wheels[WheelIdx].SkidNormal = Current.SkidNormal;
+					PVehicleOutput->Wheels[WheelIdx].SlipAngle = Current.SlipAngle;
 
 					PVehicleOutput->Wheels[WheelIdx].SuspensionOffset = Current.SuspensionOffset;
 					PVehicleOutput->Wheels[WheelIdx].SpringForce = Current.SpringForce;
diff --git a/Engine/Plugins/Experimental/ChaosVehiclesPlugin/Source/ChaosVehicles/Private/ChaosWheeledVehicleMovementComponent.cpp b/Engine/Plugins/Experimental/ChaosVehiclesPlugin/Source/ChaosVehicles/Private/ChaosWheeledVehicleMovementComponent.cpp
--- a/Engine/Plugins/Experimental/ChaosVehiclesPlugin/Source/ChaosVehicles/Private/ChaosWheeledVehicleMovementComponent.cpp
+++ b/Engine/Plugins/Experimental/ChaosVehiclesPlugin/Source/ChaosVehicles/Private/ChaosWheeledVehicleMovementComponent.cpp
@@ -68,12 +68,13 @@ FAutoConsoleCommand CVarCommandVehiclesPrevDebugPage(
 
 FString FWheelStatus::ToString() const
 {
-	return FString::Printf(TEXT("bInContact:%s ContactPoint:%s PhysMaterial:%s NormSuspensionLength:%f SpringForce:%f bIsSlipping:%s SlipMagnitude:%f bIsSkidding:%s SkidMagnitude:%f SkidNormal:%s"),
+	return FString::Printf(TEXT("bInContact:%s ContactPoint:%s PhysMaterial:%s NormSuspensionLength:%f SpringForce:%f SlipAngle:%f bIsSlipping:%s SlipMagnitude:%f bIsSkidding:%s SkidMagnitude:%f SkidNormal:%s"),
 		bInContact == true ? TEXT("True") : TEXT("False"),
 		*ContactPoint.ToString(),
 		PhysMaterial.IsValid() ? *PhysMaterial->GetName() : TEXT("None"),
 		NormalizedSuspensionLength,
 		SpringForce,
+		SlipAngle,
 		bIsSlipping == true ? TEXT("True") : TEXT("False"),
 		SlipMagnitude,
 		bIsSkidding == true ? TEXT("True") : TEXT("False"),
@@ -466,15 +467,24 @@ void UChaosWheeledVehicleSimulation::ApplyWheelFrictionForces(float DeltaTime)
 			if (GWheeledVehicleDebugParams.ShowWheelForces)
 			{
 				// show longitudinal drive force
-				FDebugDrawQueue::GetInstance().DrawDebugLine(
-						  WheelState.WheelWorldLocation[WheelIdx]
-						, WheelState.WheelWorldLocation[WheelIdx] + FrictionForceVector * 0.001f
-						, FColor::Yellow, false, -1.0f, 0, 2);
+				if (PWheel.AvailableGrip > 0.0f)
+				{
+					float Radius = 50.0f;
+					float Scaling = 50.0f / PWheel.AvailableGrip;
+
+					FVector Center = WheelState.WheelWorldLocation[WheelIdx];			
+					FVector Offset(0.0f, WheelState.WheelLocalLocation[WheelIdx].Y, 10.f);
+					Offset = Mat.TransformVector(Offset);
+
+					FDebugDrawQueue::GetInstance().DrawDebugLine(Center, Center + GroundZVector * 100.f, FColor::Orange, false, -1.0f, 0, 2);
 
-				FDebugDrawQueue::GetInstance().DrawDebugLine(
-						  WheelState.WheelWorldLocation[WheelIdx]
-						, WheelState.WheelWorldLocation[WheelIdx] + GroundZVector * 100.f
-						, FColor::Orange, false, -1.0f, 0, 2);
+					Center += Offset;
+					FVector InputForceVectorWorld = Mat.TransformVector(PWheel.InputForces);
+					FDebugDrawQueue::GetInstance().DrawDebugCircle(Center, Radius, 60, FColor::White, false, -1.0f, 0, 3, FVector(1,0,0), FVector(0,1,0), false);
+					FDebugDrawQueue::GetInstance().DrawDebugLine(Center, Center + InputForceVectorWorld * Scaling, (PWheel.bClipping?FColor::Red:FColor::Green), false, -1.0f, 0, PWheel.bClipping?2:4);
+					FDebugDrawQueue::GetInstance().DrawDebugLine(Center, Center + FrictionForceVector * Scaling, FColor::Yellow, false, -1.0f, 1, PWheel.bClipping?4:2);
+
+				}
 
 			}
 #endif
@@ -904,6 +914,7 @@ void UChaosWheeledVehicleSimulation::FillOutputState(FChaosVehicleAsyncOutput& O
 		WheelsOut.bIsSkidding = VehicleWheels[WheelIdx].IsSkidding();
 		WheelsOut.SkidMagnitude = VehicleWheels[WheelIdx].GetSkidMagnitude();
 		WheelsOut.SkidNormal = WheelState.WorldWheelVelocity[WheelIdx].GetSafeNormal();
+		WheelsOut.SlipAngle = VehicleWheels[WheelIdx].GetSlipAngle();
 
 		WheelsOut.SuspensionOffset = VehicleSuspension[WheelIdx].GetSuspensionOffset();
 		WheelsOut.SpringForce = VehicleSuspension[WheelIdx].GetSuspensionForce();
@@ -1942,6 +1953,7 @@ void UChaosWheeledVehicleMovementComponent::FillWheelOutputState()
 		State.PhysMaterial = HitResult.PhysMaterial;
 		State.NormalizedSuspensionLength = PWheel.NormalizedSuspensionLength;
 		State.SpringForce = PWheel.SpringForce;
+		State.SlipAngle = PWheel.SlipAngle;
 		State.bIsSlipping = PWheel.bIsSlipping;
 		State.SlipMagnitude = PWheel.SlipMagnitude;
 		State.bIsSkidding = PWheel.bIsSkidding;
@@ -1963,13 +1975,14 @@ void UChaosWheeledVehicleMovementComponent::FillWheelOutputState()
 
 
 void UChaosWheeledVehicleMovementComponent::BreakWheelStatus(const struct FWheelStatus& Status, bool& bInContact, FVector& ContactPoint, UPhysicalMaterial*& PhysMaterial
-	, float& NormalizedSuspensionLength, float& SpringForce, bool& bIsSlipping, float& SlipMagnitude, bool& bIsSkidding, float& SkidMagnitude, FVector& SkidNormal)
+	, float& NormalizedSuspensionLength, float& SpringForce, float& SlipAngle, bool& bIsSlipping, float& SlipMagnitude, bool& bIsSkidding, float& SkidMagnitude, FVector& SkidNormal)
 {
 	bInContact = Status.bInContact;
 	ContactPoint = Status.ContactPoint;
 	PhysMaterial = Status.PhysMaterial.Get();
 	NormalizedSuspensionLength = Status.NormalizedSuspensionLength;
 	SpringForce = Status.SpringForce;
+	SlipAngle = Status.SlipAngle;
 	bIsSlipping = Status.bIsSlipping;
 	SlipMagnitude = Status.SlipMagnitude;
 	bIsSkidding = Status.bIsSkidding;
@@ -1978,7 +1991,7 @@ void UChaosWheeledVehicleMovementComponent::BreakWheelStatus(const struct FWheel
 }
 
 FWheelStatus UChaosWheeledVehicleMovementComponent::MakeWheelStatus(bool bInContact, FVector& ContactPoint, UPhysicalMaterial* PhysMaterial
-	, float NormalizedSuspensionLength, float SpringForce, bool bIsSlipping, float SlipMagnitude, bool bIsSkidding, float SkidMagnitude, FVector& SkidNormal)
+	, float NormalizedSuspensionLength, float SpringForce, float SlipAngle, bool bIsSlipping, float SlipMagnitude, bool bIsSkidding, float SkidMagnitude, FVector& SkidNormal)
 {
 	FWheelStatus Status;
 	Status.bInContact = bInContact;
@@ -1986,6 +1999,7 @@ FWheelStatus UChaosWheeledVehicleMovementComponent::MakeWheelStatus(bool bInCont
 	Status.PhysMaterial = PhysMaterial;
 	Status.NormalizedSuspensionLength = NormalizedSuspensionLength;
 	Status.SpringForce = SpringForce;
+	Status.SlipAngle = SlipAngle;
 	Status.bIsSlipping = bIsSlipping;
 	Status.SlipMagnitude = SlipMagnitude;
 	Status.bIsSkidding = bIsSkidding;
diff --git a/Engine/Plugins/Experimental/ChaosVehiclesPlugin/Source/ChaosVehicles/Public/AsyncCallback.h b/Engine/Plugins/Experimental/ChaosVehiclesPlugin/Source/ChaosVehicles/Public/AsyncCallback.h
--- a/Engine/Plugins/Experimental/ChaosVehiclesPlugin/Source/ChaosVehicles/Public/AsyncCallback.h
+++ b/Engine/Plugins/Experimental/ChaosVehiclesPlugin/Source/ChaosVehicles/Public/AsyncCallback.h
@@ -27,6 +27,7 @@ struct FWheelsOutput
 		, WheelRadius(0.f)
 		, LateralAdhesiveLimit(0.f)
 		, LongitudinalAdhesiveLimit(0.f)
+		, SlipAngle(0.f)
 		, bIsSlipping(false)
 		, SlipMagnitude(0.f)
 		, bIsSkidding(false)
@@ -48,6 +49,7 @@ struct FWheelsOutput
 	float LateralAdhesiveLimit;
 	float LongitudinalAdhesiveLimit;
 
+	float SlipAngle;
 	bool bIsSlipping;
 	float SlipMagnitude;
 	bool bIsSkidding;
diff --git a/Engine/Plugins/Experimental/ChaosVehiclesPlugin/Source/ChaosVehicles/Public/ChaosVehicleWheel.h b/Engine/Plugins/Experimental/ChaosVehiclesPlugin/Source/ChaosVehicles/Public/ChaosVehicleWheel.h
--- a/Engine/Plugins/Experimental/ChaosVehiclesPlugin/Source/ChaosVehicles/Public/ChaosVehicleWheel.h
+++ b/Engine/Plugins/Experimental/ChaosVehiclesPlugin/Source/ChaosVehicles/Public/ChaosVehicleWheel.h
@@ -160,6 +160,9 @@ class UChaosWheeledVehicleMovementComponent;
 		UPROPERTY(EditAnywhere, Category = Wheel)
 		bool bTractionControlEnabled;
 
+		UPROPERTY(EditAnywhere, Category = Setup)
+		FRuntimeFloatCurve LateralSlipGraph;
+
 		///** Max normalized tire load at which the tire can deliver no more lateral stiffness no matter how much extra load is applied to the tire. */
 		//UPROPERTY(EditAnywhere, Category = Tire, meta = (ClampMin = "0.01", UIMin = "0.01"))
 		//float LatStiffMaxLoad;
@@ -379,6 +382,22 @@ class UChaosWheeledVehicleMovementComponent;
 			PWheelConfig.SideSlipModifier = this->SideSlipModifier;
 			PWheelConfig.SlipThreshold = this->SlipThreshold;
 			PWheelConfig.SkidThreshold = this->SkidThreshold;
+
+			PWheelConfig.LateralSlipGraph.Empty();
+			float NumSamples = 20;
+			float MinTime = 0.f, MaxTime = 0.f;
+			this->LateralSlipGraph.GetRichCurveConst()->GetTimeRange(MinTime, MaxTime);
+			if (MaxTime > 0.0f)
+			{
+				for (float X = 0; X <= MaxTime; X += (MaxTime / NumSamples))
+				{
+					float MinVal = 0.f, MaxVal = 0.f;
+					this->LateralSlipGraph.GetRichCurveConst()->GetValueRange(MinVal, MaxVal);
+					float Y = this->LateralSlipGraph.GetRichCurveConst()->Eval(X) * 10000.0f;
+					PWheelConfig.LateralSlipGraph.Add(Chaos::FVec2(X, Y));
+				}
+			}
+
 		}
 
 		void FillSuspensionSetup()
diff --git a/Engine/Plugins/Experimental/ChaosVehiclesPlugin/Source/ChaosVehicles/Public/ChaosWheeledVehicleMovementComponent.h b/Engine/Plugins/Experimental/ChaosVehiclesPlugin/Source/ChaosVehicles/Public/ChaosWheeledVehicleMovementComponent.h
--- a/Engine/Plugins/Experimental/ChaosVehiclesPlugin/Source/ChaosVehicles/Public/ChaosWheeledVehicleMovementComponent.h
+++ b/Engine/Plugins/Experimental/ChaosVehiclesPlugin/Source/ChaosVehicles/Public/ChaosWheeledVehicleMovementComponent.h
@@ -92,6 +92,10 @@ struct CHAOSVEHICLES_API FWheelStatus
 	UPROPERTY()
 	float SpringForce;
 
+	/** Slip angle at the wheel - difference between wheel local direction and velocity at wheel */
+	UPROPERTY()
+	float SlipAngle;
+
 	/** Is the wheel slipping */
 	UPROPERTY()
 	bool bIsSlipping;
@@ -100,6 +104,7 @@ struct CHAOSVEHICLES_API FWheelStatus
 	UPROPERTY()
 	float SlipMagnitude;
 
+
 	/** Is the wheel skidding */
 	UPROPERTY()
 	bool bIsSkidding;
@@ -128,6 +133,7 @@ struct CHAOSVEHICLES_API FWheelStatus
 
 	void Init()
 	{
+		SlipAngle = 0.0f;
 		bInContact = false;
 		bIsSlipping = false;
 		bIsSkidding = false;
@@ -636,11 +642,11 @@ class CHAOSVEHICLES_API UChaosWheeledVehicleMovementComponent : public UChaosVeh
 
 	UFUNCTION(BlueprintPure, Category = "Vehicles")
 	static void BreakWheelStatus(const struct FWheelStatus& Status, bool& bInContact, FVector& ContactPoint, UPhysicalMaterial*& PhysMaterial
-			, float& NormalizedSuspensionLength, float& SpringForce, bool& bIsSlipping, float& SlipMagnitude, bool& bIsSkidding, float& SkidMagnitude, FVector& SkidNormal);
+			, float& NormalizedSuspensionLength, float& SpringForce, float& SlipAngle, bool& bIsSlipping, float& SlipMagnitude, bool& bIsSkidding, float& SkidMagnitude, FVector& SkidNormal);
 
 	UFUNCTION(BlueprintPure, Category = "Vehicles")
 	static FWheelStatus MakeWheelStatus(bool bInContact, FVector& ContactPoint, UPhysicalMaterial* PhysMaterial
-			, float NormalizedSuspensionLength, float SpringForce, bool bIsSlipping, float SlipMagnitude, bool bIsSkidding, float SkidMagnitude, FVector& SkidNormal);
+			, float NormalizedSuspensionLength, float SpringForce, float SlipAngle, bool bIsSlipping, float SlipMagnitude, bool bIsSkidding, float SkidMagnitude, FVector& SkidNormal);
 
 	UFUNCTION(BlueprintPure, Category = "Vehicles")
 	static void BreakWheeledSnapshot(const struct FWheeledSnaphotData& Snapshot, FTransform& Transform, FVector& LinearVelocity
diff --git a/Engine/Source/Runtime/Experimental/ChaosVehicles/ChaosVehiclesCore/Private/VehicleUtility.cpp b/Engine/Source/Runtime/Experimental/ChaosVehicles/ChaosVehiclesCore/Private/VehicleUtility.cpp
--- a/Engine/Source/Runtime/Experimental/ChaosVehicles/ChaosVehiclesCore/Private/VehicleUtility.cpp
+++ b/Engine/Source/Runtime/Experimental/ChaosVehicles/ChaosVehiclesCore/Private/VehicleUtility.cpp
@@ -190,22 +190,32 @@ namespace Chaos
 
 	float FGraph::EvaluateY(float InX) const
 	{
-		float RangeX = BoundsX.Y - BoundsX.X;
-		float Step = (RangeX) / (Graph.Num() - 1);
-		int StartIndex = (InX - BoundsX.X) / Step;
+		float Result = 0.0f;
 
-		if (StartIndex <= 0)
+		if (!Graph.IsEmpty())
 		{
-			return Graph[0].Y;
-		}
-		else if (StartIndex >= Graph.Num() - 1)
-		{
-			return Graph[Graph.Num() - 1].Y;
-		}
+			float RangeX = BoundsX.Y - BoundsX.X;
+			float Step = (RangeX) / (Graph.Num() - 1);
+			float Start = (InX - BoundsX.X) / Step;
+
+			if (Start < 0.f)
+			{
+				return Graph[0].Y;
+			}
+			
+			int StartIndex = (int)Start;
 
-		float NormalisedRamp = (InX - Graph[StartIndex].X) / Step;
+			if (StartIndex >= Graph.Num() - 1)
+			{
+				return Graph[Graph.Num() - 1].Y;
+			}
+
+			float NormalisedRamp = (InX - Graph[StartIndex].X) / Step;
+
+			Result = Graph[StartIndex].Y * (1.f - NormalisedRamp) + Graph[StartIndex + 1].Y * NormalisedRamp;
+		}
 
-		return  Graph[StartIndex].Y * (1.f - NormalisedRamp) + Graph[StartIndex + 1].Y * NormalisedRamp;
+		return Result;
 	}
 
 } // namespace Chaos
diff --git a/Engine/Source/Runtime/Experimental/ChaosVehicles/ChaosVehiclesCore/Private/WheelSystem.cpp b/Engine/Source/Runtime/Experimental/ChaosVehicles/ChaosVehiclesCore/Private/WheelSystem.cpp
--- a/Engine/Source/Runtime/Experimental/ChaosVehicles/ChaosVehiclesCore/Private/WheelSystem.cpp
+++ b/Engine/Source/Runtime/Experimental/ChaosVehicles/ChaosVehiclesCore/Private/WheelSystem.cpp
@@ -37,6 +37,9 @@ namespace Chaos
 		, bInContact(false)
 		, WheelIndex(0)
 		, Spin(0.f)
+		, AvailableGrip(0.f)
+		, InputForces(FVector::ZeroVector)
+		, bClipping(false)
 	{
 
 	}
@@ -46,6 +49,7 @@ namespace Chaos
 		if (Setup().NewSimulationPath)
 		{
 			SimulateNew(DeltaTime);
+
 			return;
 		}
 
@@ -214,7 +218,7 @@ namespace Chaos
 		AppliedLinearBrakeForce = BrakeTorque / Re;
 
 		// Longitudinal multiplier now affecting both brake and steering equally
-		float AvailableGrip = ForceIntoSurface * SurfaceFriction * FrictionMultiplier;
+		AvailableGrip = ForceIntoSurface * SurfaceFriction * FrictionMultiplier;
 
 		float FinalLongitudinalForce = 0.f;
 		float FinalLateralForce = 0.f;
@@ -222,6 +226,7 @@ namespace Chaos
 		// currently just letting the brake override the throttle
 		bool Braking = BrakeTorque > FMath::Abs(DriveTorque);
 		bool WheelLocked = false;
+		float SlipOmega = 0.0f;
 
 		// are we actually touching the ground
 		if (ForceIntoSurface > SMALL_NUMBER)
@@ -269,18 +274,24 @@ namespace Chaos
 			float ForceRequiredToBringToStop = -(MassPerWheel * K * GroundVelocityVector.Y) / DeltaTime;
 
 			// use slip angle to generate a sideways force
-
-			// Leveling off of the graph
-			float AngleLimit = FMath::DegreesToRadians(8.0f);
-			if (SlipAngle > AngleLimit)
+			if (Setup().LateralSlipGraph.IsEmpty())
 			{
-				SlipAngle = AngleLimit;
+				float AngleLimit = FMath::DegreesToRadians(8.0f);
+				if (SlipAngle > AngleLimit)
+				{
+					SlipAngle = AngleLimit;
+				}
+				else if (SlipAngle < AngleLimit)
+				{
+					SlipAngle = AngleLimit;
+				}
+				FinalLateralForce = SlipAngle * CorneringStiffness;
 			}
-			else if (SlipAngle < AngleLimit)
-			{
-				SlipAngle = AngleLimit;
+			else
+			{ 
+				FinalLateralForce = Setup().LateralSlipGraph.EvaluateY(FMath::RadiansToDegrees(SlipAngle));
 			}
-			FinalLateralForce = SlipAngle * CorneringStiffness;
+
 			if (FinalLateralForce > FMath::Abs(ForceRequiredToBringToStop))
 			{
 				FinalLateralForce = FMath::Abs(ForceRequiredToBringToStop);
@@ -291,25 +302,34 @@ namespace Chaos
 			}
 
 			// Friction circle
+			InputForces.X = FinalLongitudinalForce;
+			InputForces.Y = FinalLateralForce;
+
 			float LengthSquared = FinalLongitudinalForce * FinalLongitudinalForce + FinalLateralForce * FinalLateralForce;
+			bClipping = false;
 			if (LengthSquared > 0.05f)
 			{
 				float Length = FMath::Sqrt(LengthSquared);
 
-				float Clip = AvailableGrip / Length;
+				float Clip = (AvailableGrip) / Length;
 				if (Clip < 1.0f)
 				{
-					FinalLongitudinalForce *= Clip;
-					FinalLateralForce *= Clip;
-
-					FinalLongitudinalForce *= Setup().SideSlipModifier;
-					FinalLateralForce *= Setup().SideSlipModifier;
-
 					if (Braking)
 					{
 						WheelLocked = true;
 					}
+					else if (FMath::Abs(FinalLongitudinalForce) > AvailableGrip)
+					{
+						SlipOmega = (FinalLongitudinalForce < 0.0f) ? -Setup().MaxSpinRotation : Setup().MaxSpinRotation;
+
+					}
+
+					bClipping = true;
+					FinalLongitudinalForce *= Clip;
+					FinalLateralForce *= Clip;
 
+					FinalLongitudinalForce *= Setup().SideSlipModifier;
+					FinalLateralForce *= Setup().SideSlipModifier;
 				}
 			}
 		}
@@ -321,7 +341,7 @@ namespace Chaos
 		else
 		{ 
 			float GroundOmega = GroundVelocityVector.X / Re;
-			Omega += (GroundOmega - Omega);
+			Omega += ((GroundOmega - Omega + SlipOmega));
 		}
 
 		// Wheel angular position
diff --git a/Engine/Source/Runtime/Experimental/ChaosVehicles/ChaosVehiclesCore/Public/VehicleUtility.h b/Engine/Source/Runtime/Experimental/ChaosVehicles/ChaosVehiclesCore/Public/VehicleUtility.h
--- a/Engine/Source/Runtime/Experimental/ChaosVehicles/ChaosVehiclesCore/Public/VehicleUtility.h
+++ b/Engine/Source/Runtime/Experimental/ChaosVehicles/ChaosVehiclesCore/Public/VehicleUtility.h
@@ -90,6 +90,8 @@ namespace Chaos
 		void Add(const FVec2& Value);
 
 		float EvaluateY(float InX) const;
+
+		bool IsEmpty() const { return Graph.IsEmpty(); }
 	private:
 		TArray<FVec2> Graph;
 		FVector2D BoundsX;
diff --git a/Engine/Source/Runtime/Experimental/ChaosVehicles/ChaosVehiclesCore/Public/WheelSystem.h b/Engine/Source/Runtime/Experimental/ChaosVehicles/ChaosVehiclesCore/Public/WheelSystem.h
--- a/Engine/Source/Runtime/Experimental/ChaosVehicles/ChaosVehiclesCore/Public/WheelSystem.h
+++ b/Engine/Source/Runtime/Experimental/ChaosVehicles/ChaosVehiclesCore/Public/WheelSystem.h
@@ -93,6 +93,7 @@ struct CHAOSVEHICLESCORE_API FSimpleWheelConfig
 		, SlipThreshold(20.0f)
 		, SkidThreshold(20.0f)
 		, NewSimulationPath(false)
+		, MaxSpinRotation(30.0f)
 	{
 
 	}
@@ -135,6 +136,9 @@ struct CHAOSVEHICLESCORE_API FSimpleWheelConfig
 
 	bool NewSimulationPath;
 
+	FGraph LateralSlipGraph;
+	float MaxSpinRotation;
+
 	// #todo: simulated Damage
 	//EWheelDamageStatus DamageStatus;
 	//float BuckleAngle;
@@ -432,6 +436,9 @@ public:
 	float LateralAdhesiveLimit;
 	float SideSlipModifier;
 	float Spin;
+	float AvailableGrip;
+	FVector InputForces;
+	bool bClipping;
 };
 
 
