Added a getter function for the bRawHandbrake member in the UChaosVehicleMovementComponent class

diff --git a/Engine/Plugins/Experimental/ChaosVehiclesPlugin/Source/ChaosVehicles/Public/ChaosVehicleMovementComponent.h b/Engine/Plugins/Experimental/ChaosVehiclesPlugin/Source/ChaosVehicles/Public/ChaosVehicleMovementComponent.h
--- a/Engine/Plugins/Experimental/ChaosVehiclesPlugin/Source/ChaosVehicles/Public/ChaosVehicleMovementComponent.h
+++ b/Engine/Plugins/Experimental/ChaosVehiclesPlugin/Source/ChaosVehicles/Public/ChaosVehicleMovementComponent.h
@@ -967,6 +967,10 @@ public:
 	UFUNCTION(BlueprintCallable, Category = "Game|Components|ChaosVehicleMovement")
 	float GetBrakeInput() { return RawBrakeInput; }
 
+	/** Get the user input for the vehicle handbrake - can use this to feed control to a connected trailer */
+	UFUNCTION(BlueprintCallable, Category = "Game|Components|ChaosVehicleMovement")
+	bool GetHandbrakeInput() const { return bRawHandbrakeInput; }
+
 	/** Get the user input for the vehicle steering - can use this to feed control to a connected trailer */
 	UFUNCTION(BlueprintCallable, Category = "Game|Components|ChaosVehicleMovement")
 	float GetSteeringInput() { return RawSteeringInput; }
