ChaosVehicles: Expose handbrake setting to BP, fix ensure in suspension when no wheels are setup on vehicle

diff --git a/Engine/Plugins/Experimental/ChaosVehiclesPlugin/Source/ChaosVehicles/Private/ChaosWheeledVehicleMovementComponent.cpp b/Engine/Plugins/Experimental/ChaosVehiclesPlugin/Source/ChaosVehicles/Private/ChaosWheeledVehicleMovementComponent.cpp
--- a/Engine/Plugins/Experimental/ChaosVehiclesPlugin/Source/ChaosVehicles/Private/ChaosWheeledVehicleMovementComponent.cpp
+++ b/Engine/Plugins/Experimental/ChaosVehiclesPlugin/Source/ChaosVehicles/Private/ChaosWheeledVehicleMovementComponent.cpp
@@ -1336,7 +1336,7 @@ void UChaosWheeledVehicleMovementComponent::SetupVehicleShapes()
 
 void UChaosWheeledVehicleMovementComponent::SetupSuspension(TUniquePtr<Chaos::FSimpleWheeledVehicle>& PVehicle)
 {
-	if (!PVehicle.IsValid())
+	if (!PVehicle.IsValid() || PVehicle->Suspension.Num() == 0)
 	{
 		return;
 	}
@@ -1360,7 +1360,7 @@ void UChaosWheeledVehicleMovementComponent::SetupSuspension(TUniquePtr<Chaos::FS
 
 	// Calculate the mass that will rest on each of the springs
 	TArray<float> OutSprungMasses;
-	if (!FSuspensionUtility::ComputeSprungMasses(LocalSpringPositions, TotalMass, OutSprungMasses))
+	if (!LocalSpringPositions.IsEmpty() && !FSuspensionUtility::ComputeSprungMasses(LocalSpringPositions, TotalMass, OutSprungMasses))
 	{
 		// if the sprung mass calc fails fall back to something that will still simulate
 		for (int Index = 0; Index < OutSprungMasses.Num(); Index++)
