ChaosVehicles: Fix nullptr crash when vehicle asset is misconfigured

diff --git a/Engine/Plugins/Experimental/ChaosVehiclesPlugin/Source/ChaosVehicles/Private/ChaosWheeledVehicleMovementComponent.cpp b/Engine/Plugins/Experimental/ChaosVehiclesPlugin/Source/ChaosVehicles/Private/ChaosWheeledVehicleMovementComponent.cpp
--- a/Engine/Plugins/Experimental/ChaosVehiclesPlugin/Source/ChaosVehicles/Private/ChaosWheeledVehicleMovementComponent.cpp
+++ b/Engine/Plugins/Experimental/ChaosVehiclesPlugin/Source/ChaosVehicles/Private/ChaosWheeledVehicleMovementComponent.cpp
@@ -1017,31 +1017,33 @@ void UChaosWheeledVehicleMovementComponent::FixupSkeletalMesh()
 					if (Wheels[WheelIdx]->bNewWheelSimulation)
 					{
 						FBodyInstance* TargetInstance = UpdatedPrimitive->GetBodyInstance();
-
-						FPhysicsCommand::ExecuteWrite(TargetInstance->ActorHandle, [&](const FPhysicsActorHandle& Chassis)
-							{
+						if (TargetInstance)
+						{
+							FPhysicsCommand::ExecuteWrite(TargetInstance->ActorHandle, [&](const FPhysicsActorHandle& Chassis)
+								{
 #if WITH_CHAOS
-								const FVector LocalWheel = GetWheelRestingPosition(WheelSetup);
-								FPhysicsConstraintHandle ConstraintHandle = FPhysicsInterface::CreateSuspension(Chassis, LocalWheel);
+									const FVector LocalWheel = GetWheelRestingPosition(WheelSetup);
+									FPhysicsConstraintHandle ConstraintHandle = FPhysicsInterface::CreateSuspension(Chassis, LocalWheel);
 
-								if (ConstraintHandle.IsValid())
-								{
-									UChaosVehicleWheel* Wheel = Wheels[WheelIdx];
-									check(Wheel);
-									ConstraintHandles.Add(ConstraintHandle);
-									if (Chaos::FSuspensionConstraint* Constraint = static_cast<Chaos::FSuspensionConstraint*>(ConstraintHandle.Constraint))
+									if (ConstraintHandle.IsValid())
 									{
-										Constraint->SetHardstopStiffness(1.0f);
-										Constraint->SetSpringStiffness(Chaos::MToCm(Wheel->SpringRate) * 0.25f);
-										Constraint->SetSpringPreload(Chaos::MToCm(Wheel->SpringPreload));
-										Constraint->SetSpringDamping(Wheel->SuspensionDampingRatio * 5.0f);
-										Constraint->SetMinLength(-Wheel->SuspensionMaxRaise);
-										Constraint->SetMaxLength(Wheel->SuspensionMaxDrop);
-										Constraint->SetAxis(-Wheel->SuspensionAxis);
+										UChaosVehicleWheel* Wheel = Wheels[WheelIdx];
+										check(Wheel);
+										ConstraintHandles.Add(ConstraintHandle);
+										if (Chaos::FSuspensionConstraint* Constraint = static_cast<Chaos::FSuspensionConstraint*>(ConstraintHandle.Constraint))
+										{
+											Constraint->SetHardstopStiffness(1.0f);
+											Constraint->SetSpringStiffness(Chaos::MToCm(Wheel->SpringRate) * 0.25f);
+											Constraint->SetSpringPreload(Chaos::MToCm(Wheel->SpringPreload));
+											Constraint->SetSpringDamping(Wheel->SuspensionDampingRatio * 5.0f);
+											Constraint->SetMinLength(-Wheel->SuspensionMaxRaise);
+											Constraint->SetMaxLength(Wheel->SuspensionMaxDrop);
+											Constraint->SetAxis(-Wheel->SuspensionAxis);
+										}
 									}
-								}
 #endif // WITH_CHAOS
-							});
+								});
+						}
 					}
 
 				}
