ChaosVehicles: Re-enable batch raycast queries, large scaled up terrains have been removed now.

diff --git a/Engine/Plugins/Experimental/ChaosVehiclesPlugin/Source/ChaosVehicles/Public/ChaosVehicleMovementComponent.h b/Engine/Plugins/Experimental/ChaosVehiclesPlugin/Source/ChaosVehicles/Public/ChaosVehicleMovementComponent.h
--- a/Engine/Plugins/Experimental/ChaosVehiclesPlugin/Source/ChaosVehicles/Public/ChaosVehicleMovementComponent.h
+++ b/Engine/Plugins/Experimental/ChaosVehiclesPlugin/Source/ChaosVehicles/Public/ChaosVehicleMovementComponent.h
@@ -89,7 +89,7 @@ struct FVehicleDebugParams
 	bool DisableAerodynamics = false;
 	bool DisableAerofoils = false;
 	bool DisableThrusters = false;
-	bool BatchQueries = false;	// Turned off due to Issue with Overlap Queries on scaled terrain
+	bool BatchQueries = true;
 	float ForceDebugScaling = 0.0006f;
 	float SleepCounterThreshold = 15;
 	bool DisableVehicleSleep = false;
