Fix CIS

diff --git a/Engine/Plugins/Experimental/ChaosVehiclesPlugin/Source/ChaosVehicles/Public/ChaosVehicleWheel.h b/Engine/Plugins/Experimental/ChaosVehiclesPlugin/Source/ChaosVehicles/Public/ChaosVehicleWheel.h
--- a/Engine/Plugins/Experimental/ChaosVehiclesPlugin/Source/ChaosVehicles/Public/ChaosVehicleWheel.h
+++ b/Engine/Plugins/Experimental/ChaosVehiclesPlugin/Source/ChaosVehicles/Public/ChaosVehicleWheel.h
@@ -14,6 +14,7 @@
 #include "SimpleVehicle.h"
 #include "VehicleUtility.h"
 #include "Engine/EngineTypes.h"
+#include "Curves/CurveFloat.h"
 
 #include "ChaosVehicleWheel.generated.h"
 
