Unshelved from pending changelist '15676712':

ChaosVehicles: Optiizations
* Simulate vehicles in parallel
* Option to cache suspension overlap test

diff --git a/Engine/Plugins/Experimental/ChaosVehiclesPlugin/Source/ChaosVehicles/Private/AsyncCallback.cpp b/Engine/Plugins/Experimental/ChaosVehiclesPlugin/Source/ChaosVehicles/Private/AsyncCallback.cpp
--- a/Engine/Plugins/Experimental/ChaosVehiclesPlugin/Source/ChaosVehicles/Private/AsyncCallback.cpp
+++ b/Engine/Plugins/Experimental/ChaosVehiclesPlugin/Source/ChaosVehicles/Private/AsyncCallback.cpp
@@ -9,6 +9,8 @@
 
 extern FVehicleDebugParams GVehicleDebugParams;
 
+DECLARE_CYCLE_STAT(TEXT("AsyncCallback:OnPreSimulate_Internal"), STAT_AsyncCallback_OnPreSimulate, STATGROUP_ChaosVehicleManager);
+
 /**
  * Callback from Physics thread
  */
@@ -16,6 +18,8 @@ void FChaosVehicleManagerAsyncCallback::OnPreSimulate_Internal()
 {
 	using namespace Chaos;
 
+	SCOPE_CYCLE_COUNTER(STAT_AsyncCallback_OnPreSimulate);
+
 	float DeltaTime = GetDeltaTime_Internal();
 	float SimTime = GetSimTime_Internal();
 
diff --git a/Engine/Plugins/Experimental/ChaosVehiclesPlugin/Source/ChaosVehicles/Private/ChaosVehicleManager.cpp b/Engine/Plugins/Experimental/ChaosVehiclesPlugin/Source/ChaosVehicles/Private/ChaosVehicleManager.cpp
--- a/Engine/Plugins/Experimental/ChaosVehiclesPlugin/Source/ChaosVehicles/Private/ChaosVehicleManager.cpp
+++ b/Engine/Plugins/Experimental/ChaosVehiclesPlugin/Source/ChaosVehicles/Private/ChaosVehicleManager.cpp
@@ -11,13 +11,10 @@
 #include "PBDRigidsSolver.h"
 #include "PhysicsProxy/SingleParticlePhysicsProxy.h"
 
+DECLARE_CYCLE_STAT(TEXT("VehicleManager:ParallelUpdateVehicles"), STAT_ChaosVehicleManager_ParallelUpdateVehicles, STATGROUP_ChaosVehicleManager);
+DECLARE_CYCLE_STAT(TEXT("VehicleManager:Update"), STAT_ChaosVehicleManager_Update, STATGROUP_ChaosVehicleManager);
+DECLARE_CYCLE_STAT(TEXT("VehicleManager:ScenePreTick"), STAT_ChaosVehicleManager_ScenePreTick, STATGROUP_ChaosVehicleManager);
 
-DECLARE_STATS_GROUP(TEXT("ChaosVehicleManager"), STATGROUP_ChaosVehicleManager, STATGROUP_Advanced);
-DECLARE_CYCLE_STAT(TEXT("VehicleSuspensionRaycasts"), STAT_ChaosVehicleManager_VehicleSuspensionRaycasts, STATGROUP_ChaosVehicleManager);
-DECLARE_CYCLE_STAT(TEXT("UpdatePhysicsVehicles"), STAT_ChaosVehicleManager_UpdatePhysicsVehicles, STATGROUP_ChaosVehicleManager);
-DECLARE_CYCLE_STAT(TEXT("TickVehicles"), STAT_ChaosVehicleManager_TickVehicles, STATGROUP_ChaosVehicleManager);
-DECLARE_CYCLE_STAT(TEXT("VehicleManagerUpdate"), STAT_ChaosVehicleManager_Update, STATGROUP_ChaosVehicleManager);
-DECLARE_CYCLE_STAT(TEXT("PretickVehicles"), STAT_ChaosVehicleManager_PretickVehicles, STATGROUP_Physics);
 
 extern FVehicleDebugParams GVehicleDebugParams;
 
@@ -155,7 +152,7 @@ void FChaosVehicleManager::RemoveVehicle(TWeakObjectPtr<UChaosVehicleMovementCom
 void FChaosVehicleManager::ScenePreTick(FPhysScene* PhysScene, float DeltaTime)
 {
 	// inputs being set via back door, i.e. accessing PVehicle directly is a no go now, needs to go through async input system
-	SCOPE_CYCLE_COUNTER(STAT_ChaosVehicleManager_PretickVehicles);
+	SCOPE_CYCLE_COUNTER(STAT_ChaosVehicleManager_ScenePreTick);
 
 	for (int32 i = 0; i < Vehicles.Num(); ++i)
 	{
@@ -166,6 +163,8 @@ void FChaosVehicleManager::ScenePreTick(FPhysScene* PhysScene, float DeltaTime)
 
 void FChaosVehicleManager::Update(FPhysScene* PhysScene, float DeltaTime)
 {
+	SCOPE_CYCLE_COUNTER(STAT_ChaosVehicleManager_Update);
+
 	UWorld* World = Scene.GetOwningWorld();
 
 	SubStepCount = 0;
@@ -187,23 +186,24 @@ void FChaosVehicleManager::Update(FPhysScene* PhysScene, float DeltaTime)
 
 void FChaosVehicleManager::SubStep(FPhysScene* PhysScene, float DeltaTime)
 {
-	if (SubStepCount++ > 0)	//ignore first substep
-	{
-		ParallelUpdateVehicles(DeltaTime);
-		//if (UWorld* World = GetWorld())
-		//{
-		//	for (TWeakObjectPtr<UChaosVehicleMovementComponent> Vehicle : Vehicles)
-		//	{
-		//		if (UPrimitiveComponent* RootPrim = Cast<UPrimitiveComponent>(Vehicle->GetRootComponent()))
-		//		{
-		//			if (RootPrim->GetBodyInstance() && !Vehicle->IsAsleep())
-		//			{
-		//				Vehicle->Substep(DeltaTime);
-		//			}
-		//		}
-		//	}
-		//}
-	}
+// #TODO: needs a callback from the physics thread during sub-step
+	//if (SubStepCount++ > 0)	//ignore first substep
+	//{
+	//	ParallelUpdateVehicles(DeltaTime);
+	//	//if (UWorld* World = GetWorld())
+	//	//{
+	//	//	for (TWeakObjectPtr<UChaosVehicleMovementComponent> Vehicle : Vehicles)
+	//	//	{
+	//	//		if (UPrimitiveComponent* RootPrim = Cast<UPrimitiveComponent>(Vehicle->GetRootComponent()))
+	//	//		{
+	//	//			if (RootPrim->GetBodyInstance() && !Vehicle->IsAsleep())
+	//	//			{
+	//	//				Vehicle->Substep(DeltaTime);
+	//	//			}
+	//	//		}
+	//	//	}
+	//	//}
+	//}
 }
 
 void FChaosVehicleManager::PostUpdate(FChaosScene* PhysScene)
@@ -213,6 +213,8 @@ void FChaosVehicleManager::PostUpdate(FChaosScene* PhysScene)
 
 void FChaosVehicleManager::ParallelUpdateVehicles(float DeltaSeconds)
 {
+	SCOPE_CYCLE_COUNTER(STAT_ChaosVehicleManager_ParallelUpdateVehicles);
+
 	FChaosVehicleManagerAsyncInput* AsyncInput = AsyncCallback->GetProducerInputData_External();
 
 	AsyncInput->Reset();	//only want latest frame's data
@@ -279,10 +281,11 @@ void FChaosVehicleManager::ParallelUpdateVehicles(float DeltaSeconds)
 	{
 		AsyncOutput = MoveTemp(AsyncOutputLatest);
 
-		for (TWeakObjectPtr<UChaosVehicleMovementComponent> Vehicle : Vehicles)
-		{
-			//Vehicle->GameThread_ProcessIntermediateAsyncOutput(*AsyncOutput);
-		}
+		// Note: not used - left as a reminder
+		//for (TWeakObjectPtr<UChaosVehicleMovementComponent> Vehicle : Vehicles)
+		//{
+		//	//Vehicle->GameThread_ProcessIntermediateAsyncOutput(*AsyncOutput);
+		//}
 	}
 
 	if (UWorld* World = Scene.GetOwningWorld())
diff --git a/Engine/Plugins/Experimental/ChaosVehiclesPlugin/Source/ChaosVehicles/Private/ChaosVehicleMovementComponent.cpp b/Engine/Plugins/Experimental/ChaosVehiclesPlugin/Source/ChaosVehicles/Private/ChaosVehicleMovementComponent.cpp
--- a/Engine/Plugins/Experimental/ChaosVehiclesPlugin/Source/ChaosVehicles/Private/ChaosVehicleMovementComponent.cpp
+++ b/Engine/Plugins/Experimental/ChaosVehiclesPlugin/Source/ChaosVehicles/Private/ChaosVehicleMovementComponent.cpp
@@ -59,10 +59,12 @@ FAutoConsoleVariableRef CVarChaosVehiclesDisableAerodynamics(TEXT("p.Vehicle.Dis
 FAutoConsoleVariableRef CVarChaosVehiclesDisableAerofoils(TEXT("p.Vehicle.DisableAerofoils"), GVehicleDebugParams.DisableAerofoils, TEXT("Enable/Disable Aerofoil Forces."));
 FAutoConsoleVariableRef CVarChaosVehiclesDisableThrusters(TEXT("p.Vehicle.DisableThrusters"), GVehicleDebugParams.DisableThrusters, TEXT("Enable/Disable Thruster Forces."));
 FAutoConsoleVariableRef CVarChaosVehiclesBatchQueries(TEXT("p.Vehicle.BatchQueries"), GVehicleDebugParams.BatchQueries, TEXT("Enable/Disable Batching Of Suspension Raycasts."));
+FAutoConsoleVariableRef CVarChaosVehiclesCacheTraceOverlap(TEXT("p.Vehicle.CacheTraceOverlap"), GVehicleDebugParams.CacheTraceOverlap, TEXT("Enable/Disable Caching Of Suspension Trace Overlap Test Optimization (only valid when BatchQueries enabled)."));
 FAutoConsoleVariableRef CVarChaosVehiclesForceDebugScaling(TEXT("p.Vehicle.SetForceDebugScaling"), GVehicleDebugParams.ForceDebugScaling, TEXT("Set Scaling For Force Visualisation."));
 FAutoConsoleVariableRef CVarChaosVehiclesSleepCounterThreshold(TEXT("p.Vehicle.SleepCounterThreshold"), GVehicleDebugParams.SleepCounterThreshold, TEXT("Set The Sleep Counter Iteration Threshold."));
 FAutoConsoleVariableRef CVarChaosVehiclesDisableVehicleSleep(TEXT("p.Vehicle.DisableVehicleSleep"), GVehicleDebugParams.DisableVehicleSleep, TEXT("Disable Vehicle Agressive Sleeping."));
 FAutoConsoleVariableRef CVarChaosVehiclesSetMaxMPH(TEXT("p.Vehicle.SetMaxMPH"), GVehicleDebugParams.SetMaxMPH, TEXT("Set a top speed in MPH (affects all vehicles)."));
+FAutoConsoleVariableRef CVarChaosVehiclesEnableMultithreading(TEXT("p.Vehicle.EnableMultithreading"), GVehicleDebugParams.EnableMultithreading, TEXT("Enable multi-threading of vehicle updates."));
 
 
 void FVehicleState::CaptureState(const FBodyInstance* TargetInstance, float GravityZ, float DeltaTime)
diff --git a/Engine/Plugins/Experimental/ChaosVehiclesPlugin/Source/ChaosVehicles/Private/ChaosWheeledVehicleMovementComponent.cpp b/Engine/Plugins/Experimental/ChaosVehiclesPlugin/Source/ChaosVehicles/Private/ChaosWheeledVehicleMovementComponent.cpp
--- a/Engine/Plugins/Experimental/ChaosVehiclesPlugin/Source/ChaosVehicles/Private/ChaosWheeledVehicleMovementComponent.cpp
+++ b/Engine/Plugins/Experimental/ChaosVehiclesPlugin/Source/ChaosVehicles/Private/ChaosWheeledVehicleMovementComponent.cpp
@@ -18,6 +18,8 @@
 #include "Chaos/PBDSuspensionConstraintData.h"
 #include "Chaos/DebugDrawQueue.h"
 
+#include "PhysicsProxy/SuspensionConstraintProxy.h"
+#include "PBDRigidsSolver.h"
 
 #if !(UE_BUILD_SHIPPING || UE_BUILD_TEST)
 #include "CanvasItem.h"
@@ -31,6 +33,14 @@ using namespace Chaos;
 PRAGMA_DISABLE_OPTIMIZATION
 #endif
 
+DECLARE_STATS_GROUP(TEXT("ChaosVehicle"), STATGROUP_ChaosVehicle, STATGROUP_Advanced);
+
+DECLARE_CYCLE_STAT(TEXT("Vehicle:SuspensionRaycasts"), STAT_ChaosVehicle_SuspensionRaycasts, STATGROUP_ChaosVehicle);
+DECLARE_CYCLE_STAT(TEXT("Vehicle:SuspensionOverlapTest"), STAT_ChaosVehicle_SuspensionOverlapTest, STATGROUP_ChaosVehicle);
+DECLARE_CYCLE_STAT(TEXT("Vehicle:SuspensionTraces"), STAT_ChaosVehicle_SuspensionTraces, STATGROUP_ChaosVehicle);
+DECLARE_CYCLE_STAT(TEXT("Vehicle:TickVehicle"), STAT_ChaosVehicle_TickVehicle, STATGROUP_ChaosVehicle);
+DECLARE_CYCLE_STAT(TEXT("Vehicle:UpdateSimulation"), STAT_ChaosVehicle_UpdateSimulation, STATGROUP_ChaosVehicle);
+
 
 FWheeledVehicleDebugParams GWheeledVehicleDebugParams;
 extern FVehicleDebugParams GVehicleDebugParams;
@@ -53,6 +63,9 @@ FAutoConsoleVariableRef CVarChaosVehiclesSteeringOverride(TEXT("p.Vehicle.Steeri
 
 FAutoConsoleVariableRef CVarChaosVehiclesResetMeasurements(TEXT("p.Vehicle.ResetMeasurements"), GWheeledVehicleDebugParams.ResetPerformanceMeasurements, TEXT("Reset Vehicle Performance Measurements."));
 
+FAutoConsoleVariableRef CVarChaosVehiclesOverlapTestExpansionXY(TEXT("p.Vehicle.OverlapTestExpansionXY"), GWheeledVehicleDebugParams.OverlapTestExpansionXY, TEXT("Raycast Overlap Test Expansion of Bounding Box in X/Y axes."));
+FAutoConsoleVariableRef CVarChaosVehiclesOverlapTestExpansionXZ(TEXT("p.Vehicle.OverlapTestExpansionZ"), GWheeledVehicleDebugParams.OverlapTestExpansionZ, TEXT("Raycast Overlap Test Expansion of Bounding Box in Z axis"));
+
 //FAutoConsoleVariableRef CVarChaosVehiclesDisableSuspensionConstraints(TEXT("p.Vehicle.DisableSuspensionConstraint"), GWheeledVehicleDebugParams.DisableSuspensionConstraint, TEXT("Enable/Disable Suspension Constraints."));
 
 FAutoConsoleCommand CVarCommandVehiclesNextDebugPage(
@@ -150,11 +163,15 @@ bool UChaosWheeledVehicleSimulation::CanSimulate() const
 
 void UChaosWheeledVehicleSimulation::TickVehicle(UWorld* WorldIn, float DeltaTime, const FChaosVehicleDefaultAsyncInput& InputData, FChaosVehicleAsyncOutput& OutputData, Chaos::FRigidBodyHandle_Internal* Handle)
 {
+	SCOPE_CYCLE_COUNTER(STAT_ChaosVehicle_TickVehicle);
+
 	UChaosVehicleSimulation::TickVehicle(WorldIn, DeltaTime, InputData, OutputData, Handle);
 }
 
 void UChaosWheeledVehicleSimulation::UpdateSimulation(float DeltaTime, const FChaosVehicleDefaultAsyncInput& InputData, Chaos::FRigidBodyHandle_Internal* Handle)
 {
+	SCOPE_CYCLE_COUNTER(STAT_ChaosVehicle_UpdateSimulation);
+
 	// Inherit common vehicle simulation stages ApplyAerodynamics, ApplyTorqueControl, etc
 	UChaosVehicleSimulation::UpdateSimulation(DeltaTime, InputData, Handle);
 
@@ -275,8 +292,25 @@ void UChaosWheeledVehicleSimulation::UpdateSimulation(float DeltaTime, const FCh
 
 }
 
+bool UChaosWheeledVehicleSimulation::ContainsTraces(const FBox& Box, const TArray<FSuspensionTrace>& SuspensionTrace)
+{
+	const FAABB3 Aabb(Box.Min, Box.Max);
+
+	for (int WheelIdx = 0; WheelIdx < Wheels.Num(); WheelIdx++)
+	{
+		if (!Aabb.Contains(SuspensionTrace[WheelIdx].Start) || !Aabb.Contains(SuspensionTrace[WheelIdx].End))
+		{
+			return false;
+		}
+	}
+
+	return true;
+}
+
 void UChaosWheeledVehicleSimulation::PerformSuspensionTraces(const TArray<FSuspensionTrace>& SuspensionTrace, FCollisionQueryParams& TraceParams, FCollisionResponseContainer& CollisionResponse)
 {
+	SCOPE_CYCLE_COUNTER(STAT_ChaosVehicle_SuspensionRaycasts);
+
 	ECollisionChannel SpringCollisionChannel = ECollisionChannel::ECC_WorldDynamic;
 	FCollisionResponseParams ResponseParams;
 	ResponseParams.CollisionResponse = CollisionResponse;
@@ -284,41 +318,47 @@ void UChaosWheeledVehicleSimulation::PerformSuspensionTraces(const TArray<FSuspe
 	// batching is about 0.5ms (25%) faster when there's 100 vehicles on a flat terrain
 	if (GVehicleDebugParams.BatchQueries)
 	{
-		FBox QueryBox;
-		for (int WheelIdx = 0; WheelIdx < Wheels.Num(); WheelIdx++)
+		if (!GVehicleDebugParams.CacheTraceOverlap || !ContainsTraces(QueryBox, SuspensionTrace))
 		{
+			SCOPE_CYCLE_COUNTER(STAT_ChaosVehicle_SuspensionOverlapTest);
 
-			FVector TraceStart = SuspensionTrace[WheelIdx].Start;
-			FVector TraceEnd = SuspensionTrace[WheelIdx].End;
+			bOverlapHit = false;
+			OverlapResults.Empty();
+			QueryBox.Init();
 
-			if (WheelIdx == 0)
-			{
-				QueryBox = FBox(TraceStart, TraceEnd);
-			}
-			else
+			//FBox QueryBox;
+			for (int WheelIdx = 0; WheelIdx < Wheels.Num(); WheelIdx++)
 			{
-				QueryBox.Min = QueryBox.Min.ComponentMin(TraceStart);
-				QueryBox.Min = QueryBox.Min.ComponentMin(TraceEnd);
-				QueryBox.Max = QueryBox.Max.ComponentMax(TraceStart);
-				QueryBox.Max = QueryBox.Max.ComponentMax(TraceEnd);
-			}
-		}
-		float OneWheelRadius = PVehicle->Wheels[0].GetEffectiveRadius(); // or wheel width
+				const FVector& TraceStart = SuspensionTrace[WheelIdx].Start;
+				const FVector& TraceEnd = SuspensionTrace[WheelIdx].End;
 
-		QueryBox.ExpandBy(FVector(OneWheelRadius, OneWheelRadius, OneWheelRadius)); // little extra just to be on the safe side consider 1 or 2 wheel vehicle
-		TArray<FOverlapResult> OverlapResults;
-		FCollisionShape CollisionBox;
-		CollisionBox.SetBox(QueryBox.GetExtent());
+				if (WheelIdx == 0)
+				{
+					QueryBox = FBox(TraceStart, TraceEnd);
+				}
+				else
+				{
+					QueryBox.Min = QueryBox.Min.ComponentMin(TraceStart);
+					QueryBox.Min = QueryBox.Min.ComponentMin(TraceEnd);
+					QueryBox.Max = QueryBox.Max.ComponentMax(TraceStart);
+					QueryBox.Max = QueryBox.Max.ComponentMax(TraceEnd);
+				}
+			}
+			QueryBox = QueryBox.ExpandBy(FVector(GWheeledVehicleDebugParams.OverlapTestExpansionXY, GWheeledVehicleDebugParams.OverlapTestExpansionXY, GWheeledVehicleDebugParams.OverlapTestExpansionZ));
+			FCollisionShape CollisionBox;
+			CollisionBox.SetBox(QueryBox.GetExtent());
 
-#if !(UE_BUILD_SHIPPING || UE_BUILD_TEST)
+			bOverlapHit = World->OverlapMultiByChannel(OverlapResults, QueryBox.GetCenter(), FQuat::Identity, SpringCollisionChannel, CollisionBox, TraceParams, ResponseParams);
+		}
+		
+	#if !(UE_BUILD_SHIPPING || UE_BUILD_TEST)
 		if (GWheeledVehicleDebugParams.ShowBatchQueryExtents)
 		{
 			Chaos::FDebugDrawQueue::GetInstance().DrawDebugBox(QueryBox.GetCenter(), QueryBox.GetExtent(), FQuat::Identity, FColor::Yellow, false, -1.0f, 0, 2.0f);
 		}
-#endif
-
-		const bool bOverlapHit = World->OverlapMultiByChannel(OverlapResults, QueryBox.GetCenter(), FQuat::Identity, SpringCollisionChannel, CollisionBox, TraceParams, ResponseParams);
+	#endif
 
+		SCOPE_CYCLE_COUNTER(STAT_ChaosVehicle_SuspensionTraces);
 		for (int32 WheelIdx = 0; WheelIdx < Wheels.Num(); ++WheelIdx)
 		{
 			FHitResult& HitResult = Wheels[WheelIdx]->HitResult;
@@ -382,6 +422,7 @@ void UChaosWheeledVehicleSimulation::PerformSuspensionTraces(const TArray<FSuspe
 	}
 	else
 	{
+		SCOPE_CYCLE_COUNTER(STAT_ChaosVehicle_SuspensionTraces);
 		for (int WheelIdx = 0; WheelIdx < Wheels.Num(); WheelIdx++)
 		{
 			FHitResult& HitResult = Wheels[WheelIdx]->HitResult;
@@ -524,14 +565,16 @@ void UChaosWheeledVehicleSimulation::ApplySuspensionForces(float DeltaTime)
 				{
 					if (Chaos::FSuspensionConstraint* Constraint = static_cast<Chaos::FSuspensionConstraint*>(ConstraintHandle.Constraint))
 					{
-						FVec3 P = HitResult.ImpactPoint + (PWheel.GetEffectiveRadius() * VehicleState.VehicleUpAxis);
-						
-						// This isn't threadsafe because of FDirtySet ProxiesData.Add(Base);
-						Constraint->SetTarget(P);
-						Constraint->SetEnabled(PWheel.InContact());
+						if (FSuspensionConstraintPhysicsProxy* Proxy = Constraint->GetProxy<FSuspensionConstraintPhysicsProxy>())
+						{
+							const FVec3 TargetPos = HitResult.ImpactPoint + (PWheel.GetEffectiveRadius() * VehicleState.VehicleUpAxis);
+
+							Chaos::FPhysicsSolver* Solver = Proxy->GetSolver<Chaos::FPhysicsSolver>();
+
+							Solver->SetSuspensionTargetOnPhysicsThread(Constraint, TargetPos, PWheel.InContact());
+						}
 					}
 				}
-
 			}
 #endif // WITH_CHAOS
 		}
diff --git a/Engine/Plugins/Experimental/ChaosVehiclesPlugin/Source/ChaosVehicles/Public/AsyncCallback.h b/Engine/Plugins/Experimental/ChaosVehiclesPlugin/Source/ChaosVehicles/Public/AsyncCallback.h
--- a/Engine/Plugins/Experimental/ChaosVehiclesPlugin/Source/ChaosVehicles/Public/AsyncCallback.h
+++ b/Engine/Plugins/Experimental/ChaosVehiclesPlugin/Source/ChaosVehicles/Public/AsyncCallback.h
@@ -11,6 +11,8 @@
 
 class UChaosVehicleMovementComponent;
 
+DECLARE_STATS_GROUP(TEXT("ChaosVehicleManager"), STATGROUP_ChaosVehicleManager, STATGROUP_Advanced);
+
 enum EChaosAsyncVehicleDataType : int8
 {
 	AsyncInvalid,
diff --git a/Engine/Plugins/Experimental/ChaosVehiclesPlugin/Source/ChaosVehicles/Public/ChaosVehicleMovementComponent.h b/Engine/Plugins/Experimental/ChaosVehiclesPlugin/Source/ChaosVehicles/Public/ChaosVehicleMovementComponent.h
--- a/Engine/Plugins/Experimental/ChaosVehiclesPlugin/Source/ChaosVehicles/Public/ChaosVehicleMovementComponent.h
+++ b/Engine/Plugins/Experimental/ChaosVehiclesPlugin/Source/ChaosVehicles/Public/ChaosVehicleMovementComponent.h
@@ -90,10 +90,11 @@ struct FVehicleDebugParams
 	bool DisableAerofoils = false;
 	bool DisableThrusters = false;
 	bool BatchQueries = true;
+	bool CacheTraceOverlap = false;
 	float ForceDebugScaling = 0.0006f;
 	float SleepCounterThreshold = 15;
 	bool DisableVehicleSleep = false;
-	bool EnableMultithreading = false; // possibly not all entirely thread safe at present
+	bool EnableMultithreading = true;
 	float SetMaxMPH = 0.0f;
 };
 
diff --git a/Engine/Plugins/Experimental/ChaosVehiclesPlugin/Source/ChaosVehicles/Public/ChaosWheeledVehicleMovementComponent.h b/Engine/Plugins/Experimental/ChaosVehiclesPlugin/Source/ChaosVehicles/Public/ChaosWheeledVehicleMovementComponent.h
--- a/Engine/Plugins/Experimental/ChaosVehiclesPlugin/Source/ChaosVehicles/Public/ChaosWheeledVehicleMovementComponent.h
+++ b/Engine/Plugins/Experimental/ChaosVehiclesPlugin/Source/ChaosVehicles/Public/ChaosWheeledVehicleMovementComponent.h
@@ -34,6 +34,9 @@ struct FWheeledVehicleDebugParams
 
 	bool ResetPerformanceMeasurements = false;
 
+	float OverlapTestExpansionXY = 100.f;
+	float OverlapTestExpansionZ = 50.f;
+
 	//bool DisableSuspensionConstraint = false;
 };
 
@@ -521,9 +524,9 @@ class UChaosWheeledVehicleSimulation : public UChaosVehicleSimulation
 public:
 
 	UChaosWheeledVehicleSimulation(TArray<class UChaosVehicleWheel*>& WheelsIn)
-		: Wheels(WheelsIn)
+		: Wheels(WheelsIn), bOverlapHit(false)
 	{
-
+		QueryBox.Init();
 	}
 
 	virtual ~UChaosWheeledVehicleSimulation()
@@ -569,6 +572,7 @@ public:
 	virtual void ApplySuspensionForces(float DeltaTime);
 
 	bool IsWheelSpinning() const;
+	bool ContainsTraces(const FBox& Box, const TArray<struct Chaos::FSuspensionTrace>& SuspensionTrace);
 
 
 	/** Draw 3D debug lines and things along side the 3D model */
@@ -581,6 +585,11 @@ public:
 
 	TArray<FPhysicsConstraintHandle> ConstraintHandles;
 
+	// cache trace overlap query
+	TArray<FOverlapResult> OverlapResults;
+	bool bOverlapHit;
+	FBox QueryBox;
+
 };
 
 //////////////////////////////////////////////////////////////////////////
diff --git a/Engine/Source/Runtime/Experimental/Chaos/Private/PBDRigidsSolver.cpp b/Engine/Source/Runtime/Experimental/Chaos/Private/PBDRigidsSolver.cpp
--- a/Engine/Source/Runtime/Experimental/Chaos/Private/PBDRigidsSolver.cpp
+++ b/Engine/Source/Runtime/Experimental/Chaos/Private/PBDRigidsSolver.cpp
@@ -629,6 +629,15 @@ namespace Chaos
 			});
 	}
 
+	template <typename Traits>
+	void TPBDRigidsSolver<Traits>::SetSuspensionTargetOnPhysicsThread(Chaos::FSuspensionConstraint* GTConstraint, const FVector& TargetPos, bool Enabled)
+	{
+		ensure(IsInPhysicsThreadContext());
+		FSuspensionConstraintPhysicsProxy* SuspensionProxy = GTConstraint->GetProxy<FSuspensionConstraintPhysicsProxy>();
+		check(SuspensionProxy);
+		SuspensionProxy->UpdateTargetOnPhysicsThread(this, TargetPos, Enabled);
+	}
+
 	CHAOS_API int32 RewindCaptureNumFrames = -1;
 	FAutoConsoleVariableRef CVarRewindCaptureNumFrames(TEXT("p.RewindCaptureNumFrames"),RewindCaptureNumFrames,TEXT("The number of frames to capture rewind for. Requires restart of solver"));
 
diff --git a/Engine/Source/Runtime/Experimental/Chaos/Private/PhysicsProxy/SuspensionConstraintProxy.cpp b/Engine/Source/Runtime/Experimental/Chaos/Private/PhysicsProxy/SuspensionConstraintProxy.cpp
--- a/Engine/Source/Runtime/Experimental/Chaos/Private/PhysicsProxy/SuspensionConstraintProxy.cpp
+++ b/Engine/Source/Runtime/Experimental/Chaos/Private/PhysicsProxy/SuspensionConstraintProxy.cpp
@@ -231,6 +231,17 @@ void TSuspensionConstraintProxy<Chaos::FSuspensionConstraint>::DestroyOnPhysicsT
 	}
 }
 
+template < >
+template < class Trait >
+void TSuspensionConstraintProxy<Chaos::FSuspensionConstraint>::UpdateTargetOnPhysicsThread(Chaos::TPBDRigidsSolver<Trait>* RBDSolver, const FVector& TargetPos, bool Enabled)
+{
+	if (Handle)
+	{
+		auto& SuspensionConstraints = RBDSolver->GetSuspensionConstraints();
+		SuspensionConstraints.GetSettings(Handle->GetConstraintIndex()).Target = TargetPos;
+		SuspensionConstraints.GetSettings(Handle->GetConstraintIndex()).Enabled = Enabled;
+	}
+}
 
 template class TSuspensionConstraintProxy< Chaos::FSuspensionConstraint >;
 
@@ -239,6 +250,7 @@ template void TSuspensionConstraintProxy<Chaos::FSuspensionConstraint>::Initiali
 template void TSuspensionConstraintProxy<Chaos::FSuspensionConstraint>::PushStateOnGameThread(Chaos::TPBDRigidsSolver<Chaos::Traits>* InSolver);\
 template void TSuspensionConstraintProxy<Chaos::FSuspensionConstraint>::PushStateOnPhysicsThread(Chaos::TPBDRigidsSolver<Chaos::Traits>* InSolver);\
 template void TSuspensionConstraintProxy<Chaos::FSuspensionConstraint>::DestroyOnPhysicsThread(Chaos::TPBDRigidsSolver<Chaos::Traits>* RBDSolver);\
+template void TSuspensionConstraintProxy<Chaos::FSuspensionConstraint>::UpdateTargetOnPhysicsThread(Chaos::TPBDRigidsSolver<Chaos::Traits>* RBDSolver, const FVector& TargetPos, bool Enabled);\
 
 #include "Chaos/EvolutionTraits.inl"
 #undef EVOLUTION_TRAIT
diff --git a/Engine/Source/Runtime/Experimental/Chaos/Public/Chaos/PBDSuspensionConstraints.h b/Engine/Source/Runtime/Experimental/Chaos/Public/Chaos/PBDSuspensionConstraints.h
--- a/Engine/Source/Runtime/Experimental/Chaos/Public/Chaos/PBDSuspensionConstraints.h
+++ b/Engine/Source/Runtime/Experimental/Chaos/Public/Chaos/PBDSuspensionConstraints.h
@@ -120,6 +120,10 @@ namespace Chaos
 			ConstraintSettings[ConstraintIndex] = Settings;
 		}
 
+		void SetTarget(int32 ConstraintIndex, const FVector& TargetPos)
+		{
+			ConstraintSettings[ConstraintIndex].Target = TargetPos;
+		}
 
 		FHandles& GetConstraintHandles()
 		{
diff --git a/Engine/Source/Runtime/Experimental/Chaos/Public/PBDRigidsSolver.h b/Engine/Source/Runtime/Experimental/Chaos/Public/PBDRigidsSolver.h
--- a/Engine/Source/Runtime/Experimental/Chaos/Public/PBDRigidsSolver.h
+++ b/Engine/Source/Runtime/Experimental/Chaos/Public/PBDRigidsSolver.h
@@ -146,6 +146,7 @@ namespace Chaos
 		void RegisterObject(Chaos::FSuspensionConstraint* GTConstraint);
 		void UnregisterObject(Chaos::FSuspensionConstraint* GTConstraint);
 
+		void SetSuspensionTargetOnPhysicsThread(Chaos::FSuspensionConstraint* GTConstraint, const FVector& TargetPos, bool Enabled);
 		void EnableRewindCapture(int32 NumFrames, bool InUseCollisionResimCache, TUniquePtr<IRewindCallback>&& RewindCallback = TUniquePtr<IRewindCallback>());
 		void SetRewindCallback(TUniquePtr<IRewindCallback>&& RewindCallback)
 		{
diff --git a/Engine/Source/Runtime/Experimental/Chaos/Public/PhysicsProxy/SuspensionConstraintProxy.h b/Engine/Source/Runtime/Experimental/Chaos/Public/PhysicsProxy/SuspensionConstraintProxy.h
--- a/Engine/Source/Runtime/Experimental/Chaos/Public/PhysicsProxy/SuspensionConstraintProxy.h
+++ b/Engine/Source/Runtime/Experimental/Chaos/Public/PhysicsProxy/SuspensionConstraintProxy.h
@@ -71,6 +71,9 @@ public:
 	template <typename Traits>
 	void CHAOS_API DestroyOnPhysicsThread(Chaos::TPBDRigidsSolver<Traits>* InSolver);
 
+	template <typename Traits>
+	void CHAOS_API UpdateTargetOnPhysicsThread(Chaos::TPBDRigidsSolver<Traits>* InSolver, const FVector& TargetPos, bool Enabled);
+
 	void SyncBeforeDestroy() {}
 	void OnRemoveFromScene() {}
 
