ChaosVehicles: add option to dynamically scale the slip graph

diff --git a/Engine/Plugins/Experimental/ChaosVehiclesPlugin/Source/ChaosVehicles/Private/ChaosWheeledVehicleMovementComponent.cpp b/Engine/Plugins/Experimental/ChaosVehiclesPlugin/Source/ChaosVehicles/Private/ChaosWheeledVehicleMovementComponent.cpp
--- a/Engine/Plugins/Experimental/ChaosVehiclesPlugin/Source/ChaosVehicles/Private/ChaosWheeledVehicleMovementComponent.cpp
+++ b/Engine/Plugins/Experimental/ChaosVehiclesPlugin/Source/ChaosVehicles/Private/ChaosWheeledVehicleMovementComponent.cpp
@@ -2383,6 +2383,22 @@ void UChaosWheeledVehicleMovementComponent::SetWheelFrictionMultiplier(int Wheel
 	}
 }
 
+void UChaosWheeledVehicleMovementComponent::SetWheelSlipGraphMultiplier(int WheelIndex, float Multiplier)
+{
+	if (FBodyInstance* TargetInstance = GetBodyInstance())
+	{
+		FPhysicsCommand::ExecuteWrite(TargetInstance->ActorHandle, [&](const FPhysicsActorHandle& Chassis)
+			{
+				if (VehicleSimulationPT && VehicleSimulationPT->PVehicle && WheelIndex < VehicleSimulationPT->PVehicle->Wheels.Num())
+				{
+					Chaos::FSimpleWheelSim& VehicleWheel = VehicleSimulationPT->PVehicle->Wheels[WheelIndex];
+
+					VehicleWheel.LateralSlipGraphMultiplier = Multiplier;
+				}
+			});
+	}
+}
+
 void UChaosWheeledVehicleMovementComponent::SetWheelMaxBrakeTorque(int WheelIndex, float Torque)
 {
 	if (FBodyInstance* TargetInstance = GetBodyInstance())
diff --git a/Engine/Plugins/Experimental/ChaosVehiclesPlugin/Source/ChaosVehicles/Public/ChaosWheeledVehicleMovementComponent.h b/Engine/Plugins/Experimental/ChaosVehiclesPlugin/Source/ChaosVehicles/Public/ChaosWheeledVehicleMovementComponent.h
--- a/Engine/Plugins/Experimental/ChaosVehiclesPlugin/Source/ChaosVehicles/Public/ChaosWheeledVehicleMovementComponent.h
+++ b/Engine/Plugins/Experimental/ChaosVehiclesPlugin/Source/ChaosVehicles/Public/ChaosWheeledVehicleMovementComponent.h
@@ -793,6 +793,9 @@ class CHAOSVEHICLES_API UChaosWheeledVehicleMovementComponent : public UChaosVeh
 	UFUNCTION(BlueprintCallable, Category = "Game|Components|ChaosWheeledVehicleMovement")
 	void SetWheelFrictionMultiplier(int WheelIndex, float Friction);
 
+	UFUNCTION(BlueprintCallable, Category = "Game|Components|ChaosWheeledVehicleMovement")
+	void SetWheelSlipGraphMultiplier(int WheelIndex, float Multiplier);
+
 	UFUNCTION(BlueprintCallable, Category = "Game|Components|ChaosWheeledVehicleMovement")
 	void SetWheelMaxBrakeTorque(int WheelIndex, float Torque);
 
diff --git a/Engine/Source/Runtime/Experimental/ChaosVehicles/ChaosVehiclesCore/Private/WheelSystem.cpp b/Engine/Source/Runtime/Experimental/ChaosVehicles/ChaosVehiclesCore/Private/WheelSystem.cpp
--- a/Engine/Source/Runtime/Experimental/ChaosVehicles/ChaosVehiclesCore/Private/WheelSystem.cpp
+++ b/Engine/Source/Runtime/Experimental/ChaosVehicles/ChaosVehiclesCore/Private/WheelSystem.cpp
@@ -16,6 +16,7 @@ namespace Chaos
 		, TractionControlEnabled(SetupIn->TractionControlEnabled)
 		, ABSEnabled(SetupIn->ABSEnabled)
 		, FrictionMultiplier(SetupIn->FrictionMultiplier)
+		, LateralSlipGraphMultiplier(SetupIn->LateralSlipGraphMultiplier)
 		, CorneringStiffness(SetupIn->CorneringStiffness)
 		, MaxSteeringAngle(SetupIn->MaxSteeringAngle)
 		, MaxBrakeTorque(SetupIn->MaxBrakeTorque)
@@ -289,7 +290,7 @@ namespace Chaos
 			}
 			else
 			{ 
-				FinalLateralForce = Setup().LateralSlipGraph.EvaluateY(FMath::RadiansToDegrees(SlipAngle));
+				FinalLateralForce = Setup().LateralSlipGraph.EvaluateY(FMath::RadiansToDegrees(SlipAngle)) * LateralSlipGraphMultiplier;
 			}
 
 			if (FinalLateralForce > FMath::Abs(ForceRequiredToBringToStop))
diff --git a/Engine/Source/Runtime/Experimental/ChaosVehicles/ChaosVehiclesCore/Public/WheelSystem.h b/Engine/Source/Runtime/Experimental/ChaosVehicles/ChaosVehiclesCore/Public/WheelSystem.h
--- a/Engine/Source/Runtime/Experimental/ChaosVehicles/ChaosVehiclesCore/Public/WheelSystem.h
+++ b/Engine/Source/Runtime/Experimental/ChaosVehicles/ChaosVehiclesCore/Public/WheelSystem.h
@@ -88,6 +88,7 @@ struct CHAOSVEHICLESCORE_API FSimpleWheelConfig
 		, LateralFrictionMultiplier(1.0f)
 		, LongitudinalFrictionMultiplier(1.0f)
 		, FrictionMultiplier(2.0f)
+		, LateralSlipGraphMultiplier(1.0f)
 		, CorneringStiffness(1000.0f)
 		, SideSlipModifier(1.0f)
 		, SlipThreshold(20.0f)
@@ -128,6 +129,7 @@ struct CHAOSVEHICLESCORE_API FSimpleWheelConfig
 	float LateralFrictionMultiplier;
 	float LongitudinalFrictionMultiplier;
 	float FrictionMultiplier;
+	float LateralSlipGraphMultiplier;
 	float CorneringStiffness;
 	float SideSlipModifier;
 
@@ -400,6 +402,7 @@ public:
 	bool TractionControlEnabled;// Straight Line Traction Control
 	bool ABSEnabled;			// Advanced braking system operational
 	float FrictionMultiplier;
+	float LateralSlipGraphMultiplier;
 	float CorneringStiffness;
 	float MaxSteeringAngle;
 	float MaxBrakeTorque;
