ChaosVehicles: Fix auto change variable that was not initialised with property

diff --git a/Engine/Plugins/Experimental/ChaosVehiclesPlugin/Source/ChaosVehicles/Private/ChaosWheeledVehicleMovementComponent.cpp b/Engine/Plugins/Experimental/ChaosVehiclesPlugin/Source/ChaosVehicles/Private/ChaosWheeledVehicleMovementComponent.cpp
--- a/Engine/Plugins/Experimental/ChaosVehiclesPlugin/Source/ChaosVehicles/Private/ChaosWheeledVehicleMovementComponent.cpp
+++ b/Engine/Plugins/Experimental/ChaosVehiclesPlugin/Source/ChaosVehicles/Private/ChaosWheeledVehicleMovementComponent.cpp
@@ -1343,6 +1343,7 @@ void UChaosWheeledVehicleMovementComponent::SetupVehicle(TUniquePtr<Chaos::FSimp
 
 		Chaos::FSimpleTransmissionSim TransmissionSim(&TransmissionSetup.GetPhysicsTransmissionConfig());
 		PVehicle->Transmission.Add(TransmissionSim);
+		TransmissionType = TransmissionSim.Setup().TransmissionType; // current transmission mode - dynamically modifiable at runtime
 
 		Chaos::FSimpleDifferentialSim DifferentialSim(&DifferentialSetup.GetPhysicsDifferentialConfig());
 		PVehicle->Differential.Add(DifferentialSim);
