Fix threading issue when wheel class is swapped out, also re-init suspension along with the wheel setup

diff --git a/Engine/Plugins/Experimental/ChaosVehiclesPlugin/Source/ChaosVehicles/Private/ChaosVehicleMovementComponent.cpp b/Engine/Plugins/Experimental/ChaosVehiclesPlugin/Source/ChaosVehicles/Private/ChaosVehicleMovementComponent.cpp
--- a/Engine/Plugins/Experimental/ChaosVehiclesPlugin/Source/ChaosVehicles/Private/ChaosVehicleMovementComponent.cpp
+++ b/Engine/Plugins/Experimental/ChaosVehiclesPlugin/Source/ChaosVehicles/Private/ChaosVehicleMovementComponent.cpp
@@ -600,9 +600,17 @@ void UChaosVehicleSimulation::InitializeWheel(int WheelIndex, const Chaos::FSimp
 		PVehicle->Wheels[WheelIndex].SetupPtr = InWheelSetup;
 		PVehicle->Wheels[WheelIndex].SetWheelRadius(InWheelSetup->WheelRadius);
 	}
+}
 
+void UChaosVehicleSimulation::InitializeSuspension(int WheelIndex, const Chaos::FSimpleSuspensionConfig* InSuspensionSetup)
+{
+	if (PVehicle->IsValid() && InSuspensionSetup && WheelIndex < PVehicle->Suspension.Num())
+	{
+		PVehicle->Suspension[WheelIndex].SetupPtr = InSuspensionSetup;
+	}
 }
 
+
 /**
  * UChaosVehicleMovementComponent
  */
diff --git a/Engine/Plugins/Experimental/ChaosVehiclesPlugin/Source/ChaosVehicles/Private/ChaosWheeledVehicleMovementComponent.cpp b/Engine/Plugins/Experimental/ChaosVehiclesPlugin/Source/ChaosVehicles/Private/ChaosWheeledVehicleMovementComponent.cpp
--- a/Engine/Plugins/Experimental/ChaosVehiclesPlugin/Source/ChaosVehicles/Private/ChaosWheeledVehicleMovementComponent.cpp
+++ b/Engine/Plugins/Experimental/ChaosVehiclesPlugin/Source/ChaosVehicles/Private/ChaosWheeledVehicleMovementComponent.cpp
@@ -2120,20 +2120,28 @@ void UChaosWheeledVehicleMovementComponent::SetWheelClass(int WheelIndex, TSubcl
 
 		if (TargetInstance && WheelIndex < Wheels.Num())
 		{
-			UChaosVehicleWheel* Wheel = Wheels[WheelIndex];
-			Wheel->Shutdown();
-			Wheel = NewObject<UChaosVehicleWheel>(this, InWheelClass);
-			Wheel->Init(this, WheelIndex);
 
 			FPhysicsCommand::ExecuteWrite(TargetInstance->ActorHandle, [&](const FPhysicsActorHandle& Chassis)
 				{
 					if (VehicleSimulationPT)
 					{
-						VehicleSimulationPT->InitializeWheel(WheelIndex, &Wheel->GetPhysicsWheelConfig());
+						UChaosVehicleWheel* OldWheel = Wheels[WheelIndex];
+						UChaosVehicleWheel* NewWheel = NewObject<UChaosVehicleWheel>(this, InWheelClass);
+						NewWheel->Init(this, WheelIndex);
+
+						VehicleSimulationPT->InitializeWheel(WheelIndex, &NewWheel->GetPhysicsWheelConfig());
+						VehicleSimulationPT->InitializeSuspension(WheelIndex, &NewWheel->GetPhysicsSuspensionConfig());
+
+						Wheels[WheelIndex] = NewWheel;
+
+						OldWheel->Shutdown();
 					}
-				});
+				});			
+				
+
 		}
 	}
+
 }
 
 FWheeledSnaphotData UChaosWheeledVehicleMovementComponent::GetSnapshot() const
diff --git a/Engine/Plugins/Experimental/ChaosVehiclesPlugin/Source/ChaosVehicles/Public/ChaosVehicleMovementComponent.h b/Engine/Plugins/Experimental/ChaosVehiclesPlugin/Source/ChaosVehicles/Public/ChaosVehicleMovementComponent.h
--- a/Engine/Plugins/Experimental/ChaosVehiclesPlugin/Source/ChaosVehicles/Public/ChaosVehicleMovementComponent.h
+++ b/Engine/Plugins/Experimental/ChaosVehiclesPlugin/Source/ChaosVehicles/Public/ChaosVehicleMovementComponent.h
@@ -726,6 +726,9 @@ public:
 	/** Reinitialize a wheel at runtime */
 	void InitializeWheel(int WheelIndex, const Chaos::FSimpleWheelConfig* InWheelSetup);
 
+	/** Reinitialize the physics suspension at runtime */
+	void InitializeSuspension(int WheelIndex, const Chaos::FSimpleSuspensionConfig* InSuspensionSetup);
+
 	/** Draw debug text for the wheels and suspension */
 	virtual void DrawDebug3D();
 	UWorld* World;
