Chaos Vehicles: extra debug text; local COM exposed

diff --git a/Engine/Plugins/Experimental/ChaosVehiclesPlugin/Source/ChaosVehicles/Private/ChaosVehicleMovementComponent.cpp b/Engine/Plugins/Experimental/ChaosVehiclesPlugin/Source/ChaosVehicles/Private/ChaosVehicleMovementComponent.cpp
--- a/Engine/Plugins/Experimental/ChaosVehiclesPlugin/Source/ChaosVehicles/Private/ChaosVehicleMovementComponent.cpp
+++ b/Engine/Plugins/Experimental/ChaosVehiclesPlugin/Source/ChaosVehicles/Private/ChaosVehicleMovementComponent.cpp
@@ -1475,6 +1475,7 @@ void UChaosVehicleMovementComponent::DrawDebug(UCanvas* Canvas, float& YL, float
 		if (TargetInstance)
 		{
 			YPos += Canvas->DrawText(RenderFont, FString::Printf(TEXT("Mass (Kg): %.1f"), TargetInstance->GetBodyMass()), 4, YPos);
+			YPos += Canvas->DrawText(RenderFont, FString::Printf(TEXT("Local COM : %s"), *TargetInstance->GetMassSpaceLocal().GetTranslation().ToString()), 4, YPos);
 			YPos += Canvas->DrawText(RenderFont, FString::Printf(TEXT("Inertia : %s"), *TargetInstance->GetBodyInertiaTensor().ToString()), 4, YPos);
 		}
 
