ChaosVehicles: disable vehicle parallel update by default, possible thread safety issue as reported in jira bug

diff --git a/Engine/Plugins/Experimental/ChaosVehiclesPlugin/Source/ChaosVehicles/Public/ChaosVehicleMovementComponent.h b/Engine/Plugins/Experimental/ChaosVehiclesPlugin/Source/ChaosVehicles/Public/ChaosVehicleMovementComponent.h
--- a/Engine/Plugins/Experimental/ChaosVehiclesPlugin/Source/ChaosVehicles/Public/ChaosVehicleMovementComponent.h
+++ b/Engine/Plugins/Experimental/ChaosVehiclesPlugin/Source/ChaosVehicles/Public/ChaosVehicleMovementComponent.h
@@ -94,7 +94,7 @@ struct FVehicleDebugParams
 	float ForceDebugScaling = 0.0006f;
 	float SleepCounterThreshold = 15;
 	bool DisableVehicleSleep = false;
-	bool EnableMultithreading = true;
+	bool EnableMultithreading = false;
 	float SetMaxMPH = 0.0f;
 };
 
