ChaosVehicles: Initialization fixes

diff --git a/Engine/Plugins/Experimental/ChaosVehiclesPlugin/Source/ChaosVehicles/Public/ChaosVehicleMovementComponent.h b/Engine/Plugins/Experimental/ChaosVehiclesPlugin/Source/ChaosVehicles/Public/ChaosVehicleMovementComponent.h
--- a/Engine/Plugins/Experimental/ChaosVehiclesPlugin/Source/ChaosVehicles/Public/ChaosVehicleMovementComponent.h
+++ b/Engine/Plugins/Experimental/ChaosVehiclesPlugin/Source/ChaosVehicles/Public/ChaosVehicleMovementComponent.h
@@ -105,6 +105,20 @@ struct CHAOSVEHICLES_API FVehicleReplicatedState
 {
 	GENERATED_USTRUCT_BODY()
 
+	FVehicleReplicatedState()
+	{
+		SteeringInput = 0.f;
+		ThrottleInput = 0.f;
+		BrakeInput = 0.f;
+		PitchInput = 0.f;
+		RollInput = 0.f;
+		YawInput = 0.f;
+		HandbrakeInput = 0.f;
+		TargetGear = 0;
+		ThrottleUp = 0.f;
+		ThrottleDown = 0.f;
+	}
+
 	// input replication: steering
 	UPROPERTY()
 	float SteeringInput;
@@ -166,6 +180,11 @@ struct FVehicleTorqueControlConfig
 public:
 	GENERATED_USTRUCT_BODY()
 
+	FVehicleTorqueControlConfig()
+	{
+		InitDefaults();
+	}
+
 	/** Torque Control Enabled */
 	UPROPERTY(EditAnywhere, Category = Setup)
 	bool Enabled;
@@ -236,6 +255,11 @@ struct FVehicleTargetRotationControlConfig
 public:
 	GENERATED_USTRUCT_BODY()
 
+	FVehicleTargetRotationControlConfig()
+	{
+		InitDefaults();
+	}
+
 	/** Rotation Control Enabled */
 	UPROPERTY(EditAnywhere, Category = Setup)
 	bool Enabled;
@@ -328,6 +352,11 @@ struct FVehicleStabilizeControlConfig
 public:
 	GENERATED_USTRUCT_BODY()
 
+	FVehicleStabilizeControlConfig()
+	{
+		InitDefaults();
+	}
+
 	/** Torque Control Enabled */
 	UPROPERTY(EditAnywhere, Category = Setup)
 	bool Enabled;
@@ -542,6 +571,11 @@ struct CHAOSVEHICLES_API FVehicleAerofoilConfig
 {
 	GENERATED_USTRUCT_BODY()
 
+	FVehicleAerofoilConfig()
+	{
+		InitDefaults();
+	}
+
 	// Does this aerofoil represent a fixed spoiler, an aircraft wing, etc how is controlled.
 	UPROPERTY(EditAnywhere, Category = AerofoilSetup)
 	EVehicleAerofoilType AerofoilType;
@@ -613,6 +647,11 @@ struct FVehicleThrustConfig
 {
 	GENERATED_USTRUCT_BODY()
 
+	FVehicleThrustConfig()
+	{
+		InitDefaults();
+	}
+
 	// Does this aerofoil represent a fixed spoiler, an aircraft wing, etc how is controlled.
 	UPROPERTY(EditAnywhere, Category = ThrustSetup)
 	EVehicleThrustType ThrustType;
diff --git a/Engine/Plugins/Experimental/ChaosVehiclesPlugin/Source/ChaosVehicles/Public/ChaosWheeledVehicleMovementComponent.h b/Engine/Plugins/Experimental/ChaosVehiclesPlugin/Source/ChaosVehicles/Public/ChaosWheeledVehicleMovementComponent.h
--- a/Engine/Plugins/Experimental/ChaosVehiclesPlugin/Source/ChaosVehicles/Public/ChaosWheeledVehicleMovementComponent.h
+++ b/Engine/Plugins/Experimental/ChaosVehiclesPlugin/Source/ChaosVehicles/Public/ChaosWheeledVehicleMovementComponent.h
@@ -158,6 +158,11 @@ struct FVehicleDifferentialConfig
 {
 	GENERATED_USTRUCT_BODY()
 
+	FVehicleDifferentialConfig()
+	{
+		InitDefaults();
+	}
+
 	/** Type of differential */
 	UPROPERTY(EditAnywhere, Category=Setup)
 	EVehicleDifferential DifferentialType;
@@ -193,6 +198,11 @@ struct FVehicleEngineConfig
 {
 	GENERATED_USTRUCT_BODY()
 
+	FVehicleEngineConfig()
+	{
+		InitDefaults();
+	}
+
 	/** Torque [Normalized 0..1] for a given RPM */
 	UPROPERTY(EditAnywhere, Category = Setup)
 	FRuntimeFloatCurve TorqueCurve;
@@ -277,6 +287,11 @@ struct FVehicleTransmissionConfig
 {
 	GENERATED_USTRUCT_BODY()
 
+	FVehicleTransmissionConfig()
+	{
+		InitDefaults();
+	}
+
 	friend class UChaosVehicleWheel;
 
 	/** Whether to use automatic transmission */
@@ -407,6 +422,11 @@ struct FVehicleSteeringConfig
 {
 	GENERATED_USTRUCT_BODY()
 
+	FVehicleSteeringConfig()
+	{
+		InitDefaults();
+	}
+
 	/** Single angle : both wheels steer by the same amount
 	 *  AngleRatio   : outer wheels on corner steer less than the inner ones by set ratio 
 	 *  Ackermann	 : Ackermann steering principle is applied */
diff --git a/Engine/Plugins/Experimental/ChaosVehiclesPlugin/Source/ChaosVehicles/Public/SnapshotData.h b/Engine/Plugins/Experimental/ChaosVehiclesPlugin/Source/ChaosVehicles/Public/SnapshotData.h
--- a/Engine/Plugins/Experimental/ChaosVehiclesPlugin/Source/ChaosVehicles/Public/SnapshotData.h
+++ b/Engine/Plugins/Experimental/ChaosVehiclesPlugin/Source/ChaosVehicles/Public/SnapshotData.h
@@ -12,6 +12,13 @@ struct FBaseSnapshotData //: public UObject
 public:
 	GENERATED_BODY()
 
+	FBaseSnapshotData()
+	{
+		Transform = FTransform::Identity;
+		LinearVelocity = FVector::ZeroVector;
+		AngularVelocity = FVector::ZeroVector;
+	}
+
 	UPROPERTY()
 	FTransform Transform;		// world coords
 
@@ -27,6 +34,15 @@ USTRUCT(BlueprintType)
 struct FWheelSnapshot
 {
 	GENERATED_BODY()
+
+	FWheelSnapshot()
+	{
+		SuspensionOffset = 0.f;
+		WheelRotationAngle = 0.f;
+		SteeringAngle = 0.f;
+		WheelRadius = 0.f;
+		WheelAngularVelocity = 0.f;
+	}
 	
 	UPROPERTY()
 	float SuspensionOffset;		// suspension location
@@ -50,6 +66,12 @@ struct FWheeledSnaphotData : public FBaseSnapshotData
 public:
 	GENERATED_BODY()
 
+	FWheeledSnaphotData()
+	{
+		SelectedGear = 0;
+		EngineRPM = 0.f;
+	}
+
 	UPROPERTY()
 	int SelectedGear;		// -ve reverse gear(s), 0 neutral, +ve forward gears
 
