ChaosVehicles: Center of mass user override, debug text showing the base COM & offset COM, move over to chaos debug rendering

diff --git a/Engine/Plugins/Experimental/ChaosVehiclesPlugin/Source/ChaosVehicles/Private/ChaosVehicleMovementComponent.cpp b/Engine/Plugins/Experimental/ChaosVehiclesPlugin/Source/ChaosVehicles/Private/ChaosVehicleMovementComponent.cpp
--- a/Engine/Plugins/Experimental/ChaosVehiclesPlugin/Source/ChaosVehicles/Private/ChaosVehicleMovementComponent.cpp
+++ b/Engine/Plugins/Experimental/ChaosVehiclesPlugin/Source/ChaosVehicles/Private/ChaosVehicleMovementComponent.cpp
@@ -23,6 +23,7 @@
 #include "DisplayDebugHelpers.h"
 #include "Chaos/ChaosEngineInterface.h"
 #include "Chaos/PBDJointConstraintData.h"
+#include "Chaos/DebugDrawQueue.h"
 
 #include "ChaosVehicleManager.h"
 #include "SimpleVehicle.h"
@@ -296,11 +297,11 @@ void UChaosVehicleSimulation::ApplyAerofoilForces(float DeltaTime)
 		FVector WorldAxis = VehicleState.VehicleWorldTransform.TransformVector(FVector::CrossProduct(FVector(1, 0, 0), Aerofoil.Setup().UpAxis));
 		if (GVehicleDebugParams.ShowAerofoilSurface)
 		{
-			DrawDebugLine(World, WorldLocation - WorldAxis * 150.0f, WorldLocation + WorldAxis * 150.0f, FColor::Black, false, -1.f, 0, 5.f);
+			Chaos::FDebugDrawQueue::GetInstance().DrawDebugLine(WorldLocation - WorldAxis * 150.0f, WorldLocation + WorldAxis * 150.0f, FColor::Black, false, -1.f, 0, 5.f);
 		}
 		if (GVehicleDebugParams.ShowAerofoilForces)
 		{
-			DrawDebugLine(World, WorldLocation, WorldLocation + WorldForce * GVehicleDebugParams.ForceDebugScaling, FColor::Green, false, -1.f, 0, 16.f);
+			Chaos::FDebugDrawQueue::GetInstance().DrawDebugLine(WorldLocation, WorldLocation + WorldForce * GVehicleDebugParams.ForceDebugScaling, FColor::Green, false, -1.f, 0, 16.f);
 		}
 #endif
 	}
@@ -460,12 +461,12 @@ void UChaosVehicleSimulation::DrawDebug3D()
 	if (GVehicleDebugParams.ShowCOM)
 	{
 		const Chaos::FVec3 COMWorld = Chaos::FParticleUtilitiesGT::GetCoMWorldPosition(RigidHandle);
-		DrawDebugCoordinateSystem(World, COMWorld, FRotator(BodyTransform.GetRotation()), 200.f, false, -1.f, 0, 2.f);
+		Chaos::FDebugDrawQueue::GetInstance().DrawDebugCoordinateSystem(COMWorld, FRotator(BodyTransform.GetRotation()), 200.f, false, -1.f, 0, 2.f);
 	}
 
 	if (GVehicleDebugParams.ShowModelOrigin)
 	{
-		DrawDebugCoordinateSystem(World, BodyTransform.GetLocation(), FRotator(BodyTransform.GetRotation()), 200.f, false, -1.f, 0, 2.f);
+		Chaos::FDebugDrawQueue::GetInstance().DrawDebugCoordinateSystem(BodyTransform.GetLocation(), FRotator(BodyTransform.GetRotation()), 200.f, false, -1.f, 0, 2.f);
 	}
 #endif
 }
@@ -496,7 +497,7 @@ void UChaosVehicleSimulation::AddForce(const FVector& Force, bool bAllowSubstepp
 		if (GVehicleDebugParams.ShowAllForces)
 		{
 			FVector Position = VehicleState.VehicleWorldCOM;
-			DrawDebugDirectionalArrow(World, Position, Position + Force * GVehicleDebugParams.ForceDebugScaling
+			Chaos::FDebugDrawQueue::GetInstance().DrawDebugDirectionalArrow(Position, Position + Force * GVehicleDebugParams.ForceDebugScaling
 				, 20.f, FColor::Blue, false, 0, 0, 2.f);
 		}
 #endif	
@@ -517,7 +518,7 @@ void UChaosVehicleSimulation::AddForceAtPosition(const FVector& Force, const FVe
 #if !(UE_BUILD_SHIPPING || UE_BUILD_TEST)
 		if (GVehicleDebugParams.ShowAllForces)
 		{
-			DrawDebugDirectionalArrow(World, Position, Position + Force * GVehicleDebugParams.ForceDebugScaling
+			Chaos::FDebugDrawQueue::GetInstance().DrawDebugDirectionalArrow(Position, Position + Force * GVehicleDebugParams.ForceDebugScaling
 				, 20.f, FColor::Blue, false, 0, 0, 2.f);
 		}
 #endif
@@ -541,7 +542,7 @@ void UChaosVehicleSimulation::AddImpulse(const FVector& Impulse, bool bVelChange
 		if (GVehicleDebugParams.ShowAllForces)
 		{
 			FVector Position = VehicleState.VehicleWorldCOM;
-			DrawDebugDirectionalArrow(World, Position, Position + Impulse * GVehicleDebugParams.ForceDebugScaling
+			Chaos::FDebugDrawQueue::GetInstance().DrawDebugDirectionalArrow(Position, Position + Impulse * GVehicleDebugParams.ForceDebugScaling
 				, 20.f, FColor::Red, false, 0, 0, 2.f);
 		}
 #endif
@@ -561,7 +562,7 @@ void UChaosVehicleSimulation::AddImpulseAtPosition(const FVector& Impulse, const
 #if !(UE_BUILD_SHIPPING || UE_BUILD_TEST)
 		if (GVehicleDebugParams.ShowAllForces)
 		{
-			DrawDebugDirectionalArrow(World, Position, Position + Impulse * GVehicleDebugParams.ForceDebugScaling
+			Chaos::FDebugDrawQueue::GetInstance().DrawDebugDirectionalArrow(Position, Position + Impulse * GVehicleDebugParams.ForceDebugScaling
 				, 20.f, FColor::Red, false, 0, 0, 2.f);
 		}
 #endif
@@ -1423,6 +1424,12 @@ void UChaosVehicleMovementComponent::UpdateMassProperties(FBodyInstance* BodyIns
 				InertiaTensor.Y *= this->InertiaTensorScale.Y * MassRatio;
 				InertiaTensor.Z *= this->InertiaTensorScale.Z * MassRatio;
 
+				if (bCenterOfMassOverride)
+				{
+					FTransform COMTransform = FPhysicsInterface::GetComTransformLocal_AssumesLocked(Actor);
+					COMTransform.SetTranslation(CenterOfMassOverride + BodyInstance->COMNudge);
+					FPhysicsInterface::SetComLocalPose_AssumesLocked(Actor, COMTransform);
+				}
 				FPhysicsInterface::SetMassSpaceInertiaTensor_AssumesLocked(Actor, InertiaTensor);
 				FPhysicsInterface::SetMass_AssumesLocked(Actor, this->Mass);
 			});
@@ -1474,8 +1481,12 @@ void UChaosVehicleMovementComponent::DrawDebug(UCanvas* Canvas, float& YL, float
 
 		if (TargetInstance)
 		{
+			FVector FinalCOM = TargetInstance->GetMassSpaceLocal().GetTranslation();
+			FVector OffsetCOM = TargetInstance->COMNudge;
+			FVector BaseCOM = FinalCOM - TargetInstance->COMNudge;
 			YPos += Canvas->DrawText(RenderFont, FString::Printf(TEXT("Mass (Kg): %.1f"), TargetInstance->GetBodyMass()), 4, YPos);
-			YPos += Canvas->DrawText(RenderFont, FString::Printf(TEXT("Local COM : %s"), *TargetInstance->GetMassSpaceLocal().GetTranslation().ToString()), 4, YPos);
+			YPos += Canvas->DrawText(RenderFont, FString::Printf(TEXT("Local COM : %s"), *FinalCOM.ToString()), 4, YPos);
+			YPos += Canvas->DrawText(RenderFont, FString::Printf(TEXT("[COM Base : %s  COM Offset : %s]"), *BaseCOM.ToString(), *OffsetCOM.ToString()), 4, YPos);
 			YPos += Canvas->DrawText(RenderFont, FString::Printf(TEXT("Inertia : %s"), *TargetInstance->GetBodyInertiaTensor().ToString()), 4, YPos);
 		}
 
diff --git a/Engine/Plugins/Experimental/ChaosVehiclesPlugin/Source/ChaosVehicles/Public/ChaosVehicleMovementComponent.h b/Engine/Plugins/Experimental/ChaosVehiclesPlugin/Source/ChaosVehicles/Public/ChaosVehicleMovementComponent.h
--- a/Engine/Plugins/Experimental/ChaosVehiclesPlugin/Source/ChaosVehicles/Public/ChaosVehicleMovementComponent.h
+++ b/Engine/Plugins/Experimental/ChaosVehiclesPlugin/Source/ChaosVehicles/Public/ChaosVehicleMovementComponent.h
@@ -805,6 +805,18 @@ public:
 	UPROPERTY(EditAnywhere, Category = VehicleSetup, meta = (ClampMin = "0.01", UIMin = "0.01"))
 	float Mass;
 
+	/**
+	 * Enable to override the calculated COM position with your own fixed value - this prevents the vehicle handling changing when the asset changes
+	 */
+	UPROPERTY(EditAnywhere, Category = VehicleSetup)
+	bool bCenterOfMassOverride;
+
+	/**
+	 * The center of mass override value, this value overrides the calculated COM and the COM offset value in the mesh is also ignored.
+	 */
+	UPROPERTY(EditAnywhere, Category = VehicleSetup, meta = (EditCondition = "bCentreOfMassOverride"))
+	FVector CenterOfMassOverride;
+
 	/** Chassis width used for drag force computation (cm)*/
 	UPROPERTY(EditAnywhere, Category = VehicleSetup, meta = (ClampMin = "0.01", UIMin = "0.01"))
 	float ChassisWidth;
diff --git a/Engine/Source/Runtime/Experimental/Chaos/Public/Chaos/DebugDrawQueue.h b/Engine/Source/Runtime/Experimental/Chaos/Public/Chaos/DebugDrawQueue.h
--- a/Engine/Source/Runtime/Experimental/Chaos/Public/Chaos/DebugDrawQueue.h
+++ b/Engine/Source/Runtime/Experimental/Chaos/Public/Chaos/DebugDrawQueue.h
@@ -307,6 +307,27 @@ public:
 		}
 	}
 
+	void DrawDebugCoordinateSystem(const FVector& Position, const FRotator& AxisRot, float Scale, bool bPersistentLines = false, float LifeTime = -1.f, uint8 DepthPriority = 0, float Thickness = 0.f)
+	{
+		if (IsDebugDrawingEnabled())
+		{
+			FScopeLock Lock(&CommandQueueCS);
+
+			FRotationMatrix R(AxisRot);
+			FVector const X = R.GetScaledAxis(EAxis::X);
+			FVector const Y = R.GetScaledAxis(EAxis::Y);
+			FVector const Z = R.GetScaledAxis(EAxis::Z);
+
+			if (AcceptCommand(3, Position))
+			{
+				CommandQueue.Emplace(FLatentDrawCommand::DrawLine(Position, Position + X * Scale, FColor::Red, bPersistentLines, LifeTime, DepthPriority, Thickness));
+				CommandQueue.Emplace(FLatentDrawCommand::DrawLine(Position, Position + Y * Scale, FColor::Green, bPersistentLines, LifeTime, DepthPriority, Thickness));
+				CommandQueue.Emplace(FLatentDrawCommand::DrawLine(Position, Position + Z * Scale, FColor::Blue, bPersistentLines, LifeTime, DepthPriority, Thickness));
+			}
+		}
+	}
+
+
 	void DrawDebugSphere(FVector const& Center, float Radius, int32 Segments, const FColor& Color, bool bPersistentLines = false, float LifeTime = -1.f, uint8 DepthPriority = 0, float Thickness = 0.f)
 	{
 		if (IsDebugDrawingEnabled())
