ChaosVehicles: change name of recently added property to make it more explanatory

diff --git a/Engine/Plugins/Experimental/ChaosVehiclesPlugin/Source/ChaosVehicles/Private/ChaosVehicleMovementComponent.cpp b/Engine/Plugins/Experimental/ChaosVehiclesPlugin/Source/ChaosVehicles/Private/ChaosVehicleMovementComponent.cpp
--- a/Engine/Plugins/Experimental/ChaosVehiclesPlugin/Source/ChaosVehicles/Private/ChaosVehicleMovementComponent.cpp
+++ b/Engine/Plugins/Experimental/ChaosVehiclesPlugin/Source/ChaosVehicles/Private/ChaosVehicleMovementComponent.cpp
@@ -1424,7 +1424,7 @@ void UChaosVehicleMovementComponent::UpdateMassProperties(FBodyInstance* BodyIns
 				InertiaTensor.Y *= this->InertiaTensorScale.Y * MassRatio;
 				InertiaTensor.Z *= this->InertiaTensorScale.Z * MassRatio;
 
-				if (bCenterOfMassOverride)
+				if (bEnableCenterOfMassOverride)
 				{
 					FTransform COMTransform = FPhysicsInterface::GetComTransformLocal_AssumesLocked(Actor);
 					COMTransform.SetTranslation(CenterOfMassOverride + BodyInstance->COMNudge);
diff --git a/Engine/Plugins/Experimental/ChaosVehiclesPlugin/Source/ChaosVehicles/Public/ChaosVehicleMovementComponent.h b/Engine/Plugins/Experimental/ChaosVehiclesPlugin/Source/ChaosVehicles/Public/ChaosVehicleMovementComponent.h
--- a/Engine/Plugins/Experimental/ChaosVehiclesPlugin/Source/ChaosVehicles/Public/ChaosVehicleMovementComponent.h
+++ b/Engine/Plugins/Experimental/ChaosVehiclesPlugin/Source/ChaosVehicles/Public/ChaosVehicleMovementComponent.h
@@ -809,12 +809,12 @@ public:
 	 * Enable to override the calculated COM position with your own fixed value - this prevents the vehicle handling changing when the asset changes
 	 */
 	UPROPERTY(EditAnywhere, Category = VehicleSetup)
-	bool bCenterOfMassOverride;
+	bool bEnableCenterOfMassOverride;
 
 	/**
 	 * The center of mass override value, this value overrides the calculated COM and the COM offset value in the mesh is also ignored.
 	 */
-	UPROPERTY(EditAnywhere, Category = VehicleSetup, meta = (EditCondition = "bCentreOfMassOverride"))
+	UPROPERTY(EditAnywhere, Category = VehicleSetup, meta = (EditCondition = "bEnableCenterOfMassOverride"))
 	FVector CenterOfMassOverride;
 
 	/** Chassis width used for drag force computation (cm)*/
