ChaosVehicles: clean up code

diff --git a/Engine/Plugins/Experimental/ChaosVehiclesPlugin/Source/ChaosVehicles/Private/AsyncCallback.cpp b/Engine/Plugins/Experimental/ChaosVehiclesPlugin/Source/ChaosVehicles/Private/AsyncCallback.cpp
deleted file mode 100644
diff --git a/Engine/Plugins/Experimental/ChaosVehiclesPlugin/Source/ChaosVehicles/Private/ChaosVehicleManager.cpp b/Engine/Plugins/Experimental/ChaosVehiclesPlugin/Source/ChaosVehicles/Private/ChaosVehicleManager.cpp
--- a/Engine/Plugins/Experimental/ChaosVehiclesPlugin/Source/ChaosVehicles/Private/ChaosVehicleManager.cpp
+++ b/Engine/Plugins/Experimental/ChaosVehiclesPlugin/Source/ChaosVehicles/Private/ChaosVehicleManager.cpp
@@ -55,7 +55,6 @@ FChaosVehicleManager::FChaosVehicleManager(FPhysScene* PhysScene)
 void FChaosVehicleManager::RegisterCallbacks()
 {
 	OnPhysScenePreTickHandle = Scene.OnPhysScenePreTick.AddRaw(this, &FChaosVehicleManager::Update);
-	OnPhysSceneStepHandle = Scene.OnPhysSceneStep.AddRaw(this, &FChaosVehicleManager::SubStep);
 	OnPhysScenePostTickHandle = Scene.OnPhysScenePostTick.AddRaw(this, &FChaosVehicleManager::PostUpdate);
 
 	check(AsyncCallback == nullptr);
@@ -65,7 +64,6 @@ void FChaosVehicleManager::RegisterCallbacks()
 void FChaosVehicleManager::UnregisterCallbacks()
 {
 	Scene.OnPhysScenePreTick.Remove(OnPhysScenePreTickHandle);
-	Scene.OnPhysSceneStep.Remove(OnPhysSceneStepHandle);
 	Scene.OnPhysScenePostTick.Remove(OnPhysScenePostTickHandle);
 
 	if (AsyncCallback)
@@ -186,28 +184,6 @@ void FChaosVehicleManager::Update(FPhysScene* PhysScene, float DeltaTime)
 	}
 }
 
-void FChaosVehicleManager::SubStep(FPhysScene* PhysScene, float DeltaTime)
-{
-// #TODO: needs a callback from the physics thread during sub-step
-	//if (SubStepCount++ > 0)	//ignore first substep
-	//{
-	//	ParallelUpdateVehicles(DeltaTime);
-	//	//if (UWorld* World = GetWorld())
-	//	//{
-	//	//	for (TWeakObjectPtr<UChaosVehicleMovementComponent> Vehicle : Vehicles)
-	//	//	{
-	//	//		if (UPrimitiveComponent* RootPrim = Cast<UPrimitiveComponent>(Vehicle->GetRootComponent()))
-	//	//		{
-	//	//			if (RootPrim->GetBodyInstance() && !Vehicle->IsAsleep())
-	//	//			{
-	//	//				Vehicle->Substep(DeltaTime);
-	//	//			}
-	//	//		}
-	//	//	}
-	//	//}
-	//}
-}
-
 void FChaosVehicleManager::PostUpdate(FChaosScene* PhysScene)
 {
 	// Unused
diff --git a/Engine/Plugins/Experimental/ChaosVehiclesPlugin/Source/ChaosVehicles/Private/ChaosVehicleManagerAsyncCallback.cpp b/Engine/Plugins/Experimental/ChaosVehiclesPlugin/Source/ChaosVehicles/Private/ChaosVehicleManagerAsyncCallback.cpp
new file mode 100644
--- /dev/null
+++ b/Engine/Plugins/Experimental/ChaosVehiclesPlugin/Source/ChaosVehicles/Private/ChaosVehicleManagerAsyncCallback.cpp
@@ -0,0 +1,88 @@
+// Copyright Epic Games, Inc. All Rights Reserved.
+
+#include "ChaosVehicleManagerAsyncCallback.h"
+
+#include "ChaosVehicleMovementComponent.h"
+#include "PBDRigidsSolver.h"
+#include "Chaos/ParticleHandleFwd.h"
+#include "PhysicsProxy/SingleParticlePhysicsProxy.h"
+
+extern FVehicleDebugParams GVehicleDebugParams;
+
+DECLARE_CYCLE_STAT(TEXT("AsyncCallback:OnPreSimulate_Internal"), STAT_AsyncCallback_OnPreSimulate, STATGROUP_ChaosVehicleManager);
+
+/**
+ * Callback from Physics thread
+ */
+void FChaosVehicleManagerAsyncCallback::OnPreSimulate_Internal()
+{
+	using namespace Chaos;
+
+	SCOPE_CYCLE_COUNTER(STAT_AsyncCallback_OnPreSimulate);
+
+	float DeltaTime = GetDeltaTime_Internal();
+	float SimTime = GetSimTime_Internal();
+
+	const FChaosVehicleManagerAsyncInput* Input = GetConsumerInput_Internal();
+	if (Input == nullptr)
+	{
+		return;
+	}
+
+	const int32 NumVehicles = Input->VehicleInputs.Num();
+
+	UWorld* World = Input->World.Get();	//only safe to access for scene queries
+	if (World == nullptr || NumVehicles == 0)
+	{
+		//world is gone so don't bother, or nothing to simulate.
+		return;
+	}
+
+	Chaos::FPhysicsSolver* PhysicsSolver = static_cast<Chaos::FPhysicsSolver*>(GetSolver());
+	if (PhysicsSolver == nullptr)
+	{
+		return;
+	}
+
+	FChaosVehicleManagerAsyncOutput& Output = GetProducerOutputData_Internal();
+	Output.VehicleOutputs.AddDefaulted(NumVehicles);
+	Output.Timestamp = Input->Timestamp;
+
+	const auto& InputVehiclesBatch = Input->VehicleInputs;
+	auto& OutputVehiclesBatch = Output.VehicleOutputs;
+
+	// beware running the vehicle simulation in parallel, code must remain threadsafe
+	auto LambdaParallelUpdate = [World, DeltaTime, SimTime, &InputVehiclesBatch, &OutputVehiclesBatch](int32 Idx)
+	{
+		const FChaosVehicleAsyncInput& VehicleInput = *InputVehiclesBatch[Idx];
+
+		if (VehicleInput.Proxy == nullptr || VehicleInput.Proxy->GetPhysicsThreadAPI() == nullptr)
+		{
+			return;
+		}
+
+		auto Handle = VehicleInput.Proxy->GetPhysicsThreadAPI();
+		if (Handle->ObjectState() != Chaos::EObjectStateType::Dynamic)
+		{
+			return;
+		}
+
+		bool bWake = false;
+		OutputVehiclesBatch[Idx] = VehicleInput.Simulate(World, DeltaTime, SimTime, bWake);
+
+	};
+
+	bool ForceSingleThread = !GVehicleDebugParams.EnableMultithreading;
+	PhysicsParallelFor(OutputVehiclesBatch.Num(), LambdaParallelUpdate, ForceSingleThread);
+
+
+}
+
+/**
+ * Contact modification currently unused
+ */
+void FChaosVehicleManagerAsyncCallback::OnContactModification_Internal(const TArrayView<Chaos::FPBDCollisionConstraintHandleModification>& Modifications)
+{
+
+}
+
diff --git a/Engine/Plugins/Experimental/ChaosVehiclesPlugin/Source/ChaosVehicles/Private/ChaosVehicleMovementComponent.cpp b/Engine/Plugins/Experimental/ChaosVehiclesPlugin/Source/ChaosVehicles/Private/ChaosVehicleMovementComponent.cpp
--- a/Engine/Plugins/Experimental/ChaosVehiclesPlugin/Source/ChaosVehicles/Private/ChaosVehicleMovementComponent.cpp
+++ b/Engine/Plugins/Experimental/ChaosVehiclesPlugin/Source/ChaosVehicles/Private/ChaosVehicleMovementComponent.cpp
@@ -1670,13 +1670,13 @@ TUniquePtr<FChaosVehicleAsyncOutput> FChaosVehicleDefaultAsyncInput::Simulate(UW
 	//UE_LOG(LogChaos, Warning, TEXT("Vehicle Physics Thread Tick %f"), DeltaSeconds);
 
 	//support nullptr because it allows us to go wide on filling the async inputs
-	if (Actor.Proxy == nullptr)
+	if (Proxy == nullptr)
 	{
 		return Output;
 	}
 
 	// We now have access to the physics representation of the chassis on the physics thread async tick
-	Chaos::FRigidBodyHandle_Internal* Handle = Actor.Proxy->GetPhysicsThreadAPI();
+	Chaos::FRigidBodyHandle_Internal* Handle = Proxy->GetPhysicsThreadAPI();
 
 	// FILL OUTPUT DATA HERE THAT WILL GET PASSED BACK TO THE GAME THREAD
 	Vehicle->VehicleSimulationPT->TickVehicle(World, DeltaSeconds, *this, *Output.Get(), Handle);
@@ -1697,7 +1697,7 @@ void UChaosVehicleMovementComponent::Update(float DeltaTime)
 		{
 			if (auto Handle = BodyInstance->ActorHandle)
 			{
-				CurAsyncInput->Actor.Proxy = Handle;	// vehicles are never static
+				CurAsyncInput->Proxy = Handle;	// vehicles are never static
 
 				FChaosVehicleDefaultAsyncInput* AsyncInput = static_cast<FChaosVehicleDefaultAsyncInput*>(CurAsyncInput);
 				AsyncInput->ControlInputs.ThrottleInput = ThrottleInputRate.CalcControlFunction(ThrottleInput);
diff --git a/Engine/Plugins/Experimental/ChaosVehiclesPlugin/Source/ChaosVehicles/Public/AsyncCallback.h b/Engine/Plugins/Experimental/ChaosVehiclesPlugin/Source/ChaosVehicles/Public/AsyncCallback.h
deleted file mode 100644
diff --git a/Engine/Plugins/Experimental/ChaosVehiclesPlugin/Source/ChaosVehicles/Public/ChaosVehicleManager.h b/Engine/Plugins/Experimental/ChaosVehiclesPlugin/Source/ChaosVehicles/Public/ChaosVehicleManager.h
--- a/Engine/Plugins/Experimental/ChaosVehiclesPlugin/Source/ChaosVehicles/Public/ChaosVehicleManager.h
+++ b/Engine/Plugins/Experimental/ChaosVehiclesPlugin/Source/ChaosVehicles/Public/ChaosVehicleManager.h
@@ -8,7 +8,7 @@
 #include "Chaos/SimCallbackInput.h"
 #include "Chaos/SimCallbackObject.h"
 #include "Chaos/GeometryParticlesfwd.h"
-#include "AsyncCallback.h"
+#include "ChaosVehicleManagerAsyncCallback.h"
 
 class UChaosTireConfig;
 class UChaosVehicleMovementComponent;
@@ -50,13 +50,7 @@ public:
 	/** Detach this vehicle manager from a FPhysScene (remove delegates, remove from map etc) */
 	void DetachFromPhysScene(FPhysScene* PhysScene);
 
-	/*
-	* Flow:
-	*	Update .. SubStep .. PostUpdate
-	*/
-
 	void Update(FPhysScene* PhysScene, float DeltaTime);
-	void SubStep(FPhysScene* PhysScene, float DeltaTime);
 	void PostUpdate(FChaosScene* PhysScene);
 
 	void ParallelUpdateVehicles(float DeltaSeconds);
@@ -82,7 +76,6 @@ private:
 	TArray<TWeakObjectPtr<UChaosVehicleMovementComponent>> Vehicles;
 
 	FDelegateHandle OnPhysScenePreTickHandle;
-	FDelegateHandle OnPhysSceneStepHandle;
 	FDelegateHandle OnPhysScenePostTickHandle;
 
 	static FDelegateHandle OnPostWorldInitializationHandle;
diff --git a/Engine/Plugins/Experimental/ChaosVehiclesPlugin/Source/ChaosVehicles/Public/ChaosVehicleManagerAsyncCallback.h b/Engine/Plugins/Experimental/ChaosVehiclesPlugin/Source/ChaosVehicles/Public/ChaosVehicleManagerAsyncCallback.h
new file mode 100644
--- /dev/null
+++ b/Engine/Plugins/Experimental/ChaosVehiclesPlugin/Source/ChaosVehicles/Public/ChaosVehicleManagerAsyncCallback.h
@@ -0,0 +1,171 @@
+// Copyright Epic Games, Inc. All Rights Reserved.
+
+#pragma once
+
+#include "CoreMinimal.h"
+#include "PhysicsPublic.h"
+#include "Chaos/SimCallbackInput.h"
+#include "Chaos/SimCallbackObject.h"
+#include "Chaos/GeometryParticlesfwd.h"
+#include "PhysicsProxy/SingleParticlePhysicsProxyFwd.h"
+
+class UChaosVehicleMovementComponent;
+
+DECLARE_STATS_GROUP(TEXT("ChaosVehicleManager"), STATGROUP_ChaosVehicleManager, STATGROUP_Advanced);
+
+enum EChaosAsyncVehicleDataType : int8
+{
+	AsyncInvalid,
+	AsyncDefault,
+};
+
+struct FWheelsOutput
+{
+	FWheelsOutput()
+		: InContact(false)
+		, SteeringAngle(0.f)
+		, AngularPosition(0.f)
+		, AngularVelocity(0.f)
+		, WheelRadius(0.f)
+		, LateralAdhesiveLimit(0.f)
+		, LongitudinalAdhesiveLimit(0.f)
+		, SlipAngle(0.f)
+		, bIsSlipping(false)
+		, SlipMagnitude(0.f)
+		, bIsSkidding(false)
+		, SkidMagnitude(0.f)
+		, SkidNormal(FVector(1,0,0))
+		, SuspensionOffset(0.f)
+		, SpringForce(0.f)
+		, NormalizedSuspensionLength(0.f)
+	{
+	}
+
+	// wheels
+	bool InContact;
+	float SteeringAngle;
+	float AngularPosition;
+	float AngularVelocity;
+	float WheelRadius;
+
+	float LateralAdhesiveLimit;
+	float LongitudinalAdhesiveLimit;
+
+	float SlipAngle;
+	bool bIsSlipping;
+	float SlipMagnitude;
+	bool bIsSkidding;
+	float SkidMagnitude;
+	FVector SkidNormal;
+
+	// suspension related
+	float SuspensionOffset;
+	float SpringForce;
+	float NormalizedSuspensionLength;
+
+};
+
+/**
+ * Per Vehicle Output State from Physics Thread to Game Thread                            
+ */
+struct FPhysicsVehicleOutput
+{
+	FPhysicsVehicleOutput()
+		: CurrentGear(0)
+		, TargetGear(0)
+		, EngineRPM(0.f)
+		, EngineTorque(0.f)
+		, TransmissionRPM(0.f)
+		, TransmissionTorque(0.f)
+	{
+	}
+
+	TArray<FWheelsOutput> Wheels;
+	int32 CurrentGear;
+	int32 TargetGear;
+	float EngineRPM;
+	float EngineTorque;
+	float TransmissionRPM;
+	float TransmissionTorque;
+};
+
+/**
+ * Per Vehicle Input State from Game Thread to Physics Thread
+ */                             
+struct FChaosVehicleAsyncInput
+{
+	const EChaosAsyncVehicleDataType Type;
+	UChaosVehicleMovementComponent* Vehicle;
+	
+	FSingleParticlePhysicsProxy* Proxy;
+
+	/** 
+	* Vehicle simulation running on the Physics Thread
+	*/
+	virtual TUniquePtr<struct FChaosVehicleAsyncOutput> Simulate(UWorld* World, const float DeltaSeconds, const float TotalSeconds, bool& bWakeOut) const = 0;
+
+	FChaosVehicleAsyncInput(EChaosAsyncVehicleDataType InType = EChaosAsyncVehicleDataType::AsyncInvalid)
+		: Type(InType)
+		, Vehicle(nullptr)
+	{
+		Proxy = nullptr;	//indicates async/sync task not needed
+	}
+
+	virtual ~FChaosVehicleAsyncInput() = default;
+};
+
+struct FChaosVehicleManagerAsyncInput : public Chaos::FSimCallbackInput
+{
+	TArray<TUniquePtr<FChaosVehicleAsyncInput>> VehicleInputs;
+
+	TWeakObjectPtr<UWorld> World;
+	int32 Timestamp = INDEX_NONE;
+
+	void Reset()
+	{
+		VehicleInputs.Reset();
+		World.Reset();
+	}
+};
+
+/**
+ * Async Output Data
+ */
+struct FChaosVehicleAsyncOutput
+{
+	const EChaosAsyncVehicleDataType Type;
+	bool bValid;	// indicates no work was done
+	FPhysicsVehicleOutput VehicleSimOutput;
+
+	FChaosVehicleAsyncOutput(EChaosAsyncVehicleDataType InType = EChaosAsyncVehicleDataType::AsyncInvalid)
+		: Type(InType)
+		, bValid(false)
+	{ }
+
+	virtual ~FChaosVehicleAsyncOutput() = default;
+};
+
+
+/**
+ * Async Output for all of the vehicles handled by this Vehicle Manager
+ */
+struct FChaosVehicleManagerAsyncOutput : public Chaos::FSimCallbackOutput
+{
+	TArray<TUniquePtr<FChaosVehicleAsyncOutput>> VehicleOutputs;
+	int32 Timestamp = INDEX_NONE;
+
+	void Reset()
+	{
+		VehicleOutputs.Reset();
+	}
+};
+
+/**
+ * Async callback from the Physics Engine where we can perform our vehicle simulation
+ */
+class FChaosVehicleManagerAsyncCallback : public Chaos::TSimCallbackObject<FChaosVehicleManagerAsyncInput, FChaosVehicleManagerAsyncOutput>
+{
+private:
+	virtual void OnPreSimulate_Internal() override;
+	virtual void OnContactModification_Internal(const TArrayView<Chaos::FPBDCollisionConstraintHandleModification>& Modifications) override;
+};
diff --git a/Engine/Plugins/Experimental/ChaosVehiclesPlugin/Source/ChaosVehicles/Public/ChaosVehicleMovementComponent.h b/Engine/Plugins/Experimental/ChaosVehiclesPlugin/Source/ChaosVehicles/Public/ChaosVehicleMovementComponent.h
--- a/Engine/Plugins/Experimental/ChaosVehiclesPlugin/Source/ChaosVehicles/Public/ChaosVehicleMovementComponent.h
+++ b/Engine/Plugins/Experimental/ChaosVehiclesPlugin/Source/ChaosVehicles/Public/ChaosVehicleMovementComponent.h
@@ -14,7 +14,7 @@
 #include "ThrustSystem.h"
 #include "PhysicsProxy/SingleParticlePhysicsProxyFwd.h"
 #include "SnapshotData.h"
-#include "AsyncCallback.h"
+#include "ChaosVehicleManagerAsyncCallback.h"
 
 #include "ChaosVehicleMovementComponent.generated.h"
 
