ChaosVehicles: Additional debug visualisations, draw data relating to suspension contact surface

diff --git a/Engine/Plugins/Experimental/ChaosVehiclesPlugin/Source/ChaosVehicles/Private/ChaosVehicleWheel.cpp b/Engine/Plugins/Experimental/ChaosVehiclesPlugin/Source/ChaosVehicles/Private/ChaosVehicleWheel.cpp
--- a/Engine/Plugins/Experimental/ChaosVehiclesPlugin/Source/ChaosVehicles/Private/ChaosVehicleWheel.cpp
+++ b/Engine/Plugins/Experimental/ChaosVehiclesPlugin/Source/ChaosVehicles/Private/ChaosVehicleWheel.cpp
@@ -44,7 +44,7 @@ UChaosVehicleWheel::UChaosVehicleWheel(const FObjectInitializer& ObjectInitializ
 	SuspensionMaxRaise = 10.0f;
 	SuspensionMaxDrop = 10.0f;
 	SuspensionDampingRatio = 0.5f;
-	SuspensionSmoothing = 6;
+	SuspensionSmoothing = 0;
 	WheelLoadRatio = 0.5f;
 	RollbarScaling = 0.15f;
 	SweepType = ESweepType::SimpleSweep;
diff --git a/Engine/Plugins/Experimental/ChaosVehiclesPlugin/Source/ChaosVehicles/Private/ChaosWheeledVehicleMovementComponent.cpp b/Engine/Plugins/Experimental/ChaosVehiclesPlugin/Source/ChaosVehicles/Private/ChaosWheeledVehicleMovementComponent.cpp
--- a/Engine/Plugins/Experimental/ChaosVehiclesPlugin/Source/ChaosVehicles/Private/ChaosWheeledVehicleMovementComponent.cpp
+++ b/Engine/Plugins/Experimental/ChaosVehiclesPlugin/Source/ChaosVehicles/Private/ChaosWheeledVehicleMovementComponent.cpp
@@ -53,6 +53,7 @@ FAutoConsoleVariableRef CVarChaosVehiclesShowSuspensionLimits(TEXT("p.Vehicle.Sh
 FAutoConsoleVariableRef CVarChaosVehiclesShowWheelForces(TEXT("p.Vehicle.ShowWheelForces"), GWheeledVehicleDebugParams.ShowWheelForces, TEXT("Enable/Disable Wheel Forces Visualisation."));
 FAutoConsoleVariableRef CVarChaosVehiclesShowSuspensionForces(TEXT("p.Vehicle.ShowSuspensionForces"), GWheeledVehicleDebugParams.ShowSuspensionForces, TEXT("Enable/Disable Suspension Forces Visualisation."));
 FAutoConsoleVariableRef CVarChaosVehiclesShowBatchQueryExtents(TEXT("p.Vehicle.ShowBatchQueryExtents"), GWheeledVehicleDebugParams.ShowBatchQueryExtents, TEXT("Enable/Disable Suspension Forces Visualisation."));
+FAutoConsoleVariableRef CVarChaosVehiclesTraceTypeOverride(TEXT("p.Vehicle.TraceTypeOverride"), GWheeledVehicleDebugParams.TraceTypeOverride, TEXT("Override ray trace type, 1=Simple, 2=Complex."));
 
 FAutoConsoleVariableRef CVarChaosVehiclesDisableSuspensionForces(TEXT("p.Vehicle.DisableSuspensionForces"), GWheeledVehicleDebugParams.DisableSuspensionForces, TEXT("Enable/Disable Suspension Forces."));
 FAutoConsoleVariableRef CVarChaosVehiclesDisableFrictionForces(TEXT("p.Vehicle.DisableFrictionForces"), GWheeledVehicleDebugParams.DisableFrictionForces, TEXT("Enable/Disable Wheel Friction Forces."));
@@ -371,6 +372,11 @@ void UChaosWheeledVehicleSimulation::PerformSuspensionTraces(const TArray<FSuspe
 				const FVector& TraceEnd = SuspensionTrace[WheelIdx].End;
 				TraceParams.bTraceComplex = (Wheels[WheelIdx]->SweepType == ESweepType::ComplexSweep);
 
+				if (GWheeledVehicleDebugParams.TraceTypeOverride > 0)
+				{
+					TraceParams.bTraceComplex = GWheeledVehicleDebugParams.TraceTypeOverride == 2;
+				}
+
 				FVector TraceVector(TraceStart - TraceEnd); // reversed
 				FVector TraceNormal = TraceVector.GetSafeNormal();
 
@@ -432,6 +438,11 @@ void UChaosWheeledVehicleSimulation::PerformSuspensionTraces(const TArray<FSuspe
 			FVector TraceEnd = SuspensionTrace[WheelIdx].End;
 			TraceParams.bTraceComplex = (Wheels[WheelIdx]->SweepType == ESweepType::ComplexSweep);
 
+			if (GWheeledVehicleDebugParams.TraceTypeOverride > 0)
+			{
+				TraceParams.bTraceComplex = GWheeledVehicleDebugParams.TraceTypeOverride == 2;
+			}
+
 			FVector TraceVector(TraceStart - TraceEnd); // reversed
 			FVector TraceNormal = TraceVector.GetSafeNormal();
 
@@ -888,10 +899,32 @@ void UChaosWheeledVehicleSimulation::DrawDebug3D()
 
 	if (GWheeledVehicleDebugParams.ShowWheelCollisionNormal)
 	{
+		FString Name;
 		for (int WheelIdx = 0; WheelIdx < Wheels.Num(); WheelIdx++)
 		{
 			FHitResult& Hit = Wheels[WheelIdx]->HitResult;
-			FDebugDrawQueue::GetInstance().DrawDebugLine(Hit.ImpactPoint, Hit.ImpactPoint + Hit.Normal * 20.0f, FColor::Yellow, true, 1.0f, 0, 1.0f);
+
+			FVector VehicleRightAxis = VehicleState.VehicleWorldTransform.GetUnitAxis(EAxis::Y) * 20.0f;
+			const FVector& WheelOffset = PVehicle->Suspension[WheelIdx].GetLocalRestingPosition();
+			if (WheelOffset.Y < 0.0f)
+			{
+				VehicleRightAxis = VehicleRightAxis * -1.0f;
+			}
+
+			FVector Pt = Hit.ImpactPoint + VehicleRightAxis;
+			FDebugDrawQueue::GetInstance().DrawDebugLine(Pt, Pt + Hit.Normal * 20.0f, FColor::Yellow, false, 1.0f, 0, 1.0f);
+			FDebugDrawQueue::GetInstance().DrawDebugSphere(Pt, 5.0f, 4, FColor::White, false, 1.0f, 0, 1.0f);
+			if (Hit.GetComponent())
+			{
+				FDebugDrawQueue::GetInstance().DrawDebugString(Pt + VehicleRightAxis, Hit.GetComponent()->GetName(), nullptr, FColor::White, -1.f, true, 1.0f);
+
+				AActor* OwningActor = Hit.GetComponent()->GetOwner();
+				if (OwningActor)
+				{
+					Name = OwningActor->GetName();
+				}
+
+			}
 		}
 	}
 
diff --git a/Engine/Plugins/Experimental/ChaosVehiclesPlugin/Source/ChaosVehicles/Public/ChaosWheeledVehicleMovementComponent.h b/Engine/Plugins/Experimental/ChaosVehiclesPlugin/Source/ChaosVehicles/Public/ChaosWheeledVehicleMovementComponent.h
--- a/Engine/Plugins/Experimental/ChaosVehiclesPlugin/Source/ChaosVehicles/Public/ChaosWheeledVehicleMovementComponent.h
+++ b/Engine/Plugins/Experimental/ChaosVehiclesPlugin/Source/ChaosVehicles/Public/ChaosWheeledVehicleMovementComponent.h
@@ -24,6 +24,7 @@ struct FWheeledVehicleDebugParams
 	bool ShowWheelForces = false;
 	bool ShowSuspensionForces = false;
 	bool ShowBatchQueryExtents = false;
+	int TraceTypeOverride = 0;
 
 	bool DisableSuspensionForces = false;
 	bool DisableFrictionForces = false;
