ChaosVehicles: Vehicle parallel update fix by removing particle state change when forces applied

#okforhithub public, internal

diff --git a/Engine/Plugins/Experimental/ChaosVehiclesPlugin/Source/ChaosVehicles/Private/ChaosVehicleMovementComponent.cpp b/Engine/Plugins/Experimental/ChaosVehiclesPlugin/Source/ChaosVehicles/Private/ChaosVehicleMovementComponent.cpp
--- a/Engine/Plugins/Experimental/ChaosVehiclesPlugin/Source/ChaosVehicles/Private/ChaosVehicleMovementComponent.cpp
+++ b/Engine/Plugins/Experimental/ChaosVehiclesPlugin/Source/ChaosVehicles/Private/ChaosVehicleMovementComponent.cpp
@@ -478,17 +478,15 @@ void UChaosVehicleSimulation::AddForce(const FVector& Force, bool bAllowSubstepp
 		Chaos::EObjectStateType ObjectState = RigidHandle->ObjectState();
 		if (CHAOS_ENSURE(ObjectState == Chaos::EObjectStateType::Dynamic || ObjectState == Chaos::EObjectStateType::Sleeping))
 		{
-			RigidHandle->SetObjectState(Chaos::EObjectStateType::Dynamic);
-
 			if (bAccelChange)
 			{
 				const float RigidMass = RigidHandle->M();
 				const Chaos::FVec3 Acceleration = Force * RigidMass;
-				RigidHandle->AddForce(Acceleration);
+				RigidHandle->AddForce(Acceleration, false);
 			}
 			else
 			{
-				RigidHandle->AddForce(Force);
+				RigidHandle->AddForce(Force, false);
 			}
 
 		}
@@ -510,10 +508,9 @@ void UChaosVehicleSimulation::AddForceAtPosition(const FVector& Force, const FVe
 	if (ensure(RigidHandle))
 	{
 		const Chaos::FVec3 WorldCOM = Chaos::FParticleUtilitiesGT::GetCoMWorldPosition(RigidHandle);
-		RigidHandle->SetObjectState(Chaos::EObjectStateType::Dynamic);
 		const Chaos::FVec3 WorldTorque = Chaos::FVec3::CrossProduct(Position - WorldCOM, Force);
-		RigidHandle->AddForce(Force);
-		RigidHandle->AddTorque(WorldTorque);
+		RigidHandle->AddForce(Force, false);
+		RigidHandle->AddTorque(WorldTorque, false);
 
 #if !(UE_BUILD_SHIPPING || UE_BUILD_TEST)
 		if (GVehicleDebugParams.ShowAllForces)
@@ -531,11 +528,11 @@ void UChaosVehicleSimulation::AddImpulse(const FVector& Impulse, bool bVelChange
 	{
 		if (bVelChange)
 		{
-			RigidHandle->SetLinearImpulse(RigidHandle->LinearImpulse() + RigidHandle->M() * Impulse);
+			RigidHandle->SetLinearImpulse(RigidHandle->LinearImpulse() + RigidHandle->M() * Impulse, false);
 		}
 		else
 		{
-			RigidHandle->SetLinearImpulse(RigidHandle->LinearImpulse() + Impulse);
+			RigidHandle->SetLinearImpulse(RigidHandle->LinearImpulse() + Impulse, false);
 		}
 
 #if !(UE_BUILD_SHIPPING || UE_BUILD_TEST)
@@ -555,8 +552,8 @@ void UChaosVehicleSimulation::AddImpulseAtPosition(const FVector& Impulse, const
 	{
 		const Chaos::FVec3 WorldCOM = Chaos::FParticleUtilitiesGT::GetCoMWorldPosition(RigidHandle);
 		const Chaos::FVec3 AngularImpulse = Chaos::FVec3::CrossProduct(Position - WorldCOM, Impulse);
-		RigidHandle->SetLinearImpulse(RigidHandle->LinearImpulse() + Impulse);
-		RigidHandle->SetAngularImpulse(RigidHandle->AngularImpulse() + AngularImpulse);
+		RigidHandle->SetLinearImpulse(RigidHandle->LinearImpulse() + Impulse, false);
+		RigidHandle->SetAngularImpulse(RigidHandle->AngularImpulse() + AngularImpulse, false);
 
 
 #if !(UE_BUILD_SHIPPING || UE_BUILD_TEST)
@@ -579,18 +576,13 @@ void UChaosVehicleSimulation::AddTorqueInRadians(const FVector& Torque, bool bAl
 		{
 			if (bAccelChange)
 			{
-				RigidHandle->AddTorque(Chaos::FParticleUtilitiesXR::GetWorldInertia(RigidHandle) * Torque);
+				RigidHandle->AddTorque(Chaos::FParticleUtilitiesXR::GetWorldInertia(RigidHandle) * Torque, false);
 			}
 			else
 			{
-				RigidHandle->AddTorque(Torque);
+				RigidHandle->AddTorque(Torque, false);
 			}
 		}
-
-#if !(UE_BUILD_SHIPPING || UE_BUILD_TEST)
-		//#todo: how do we visualize torque?
-#endif
-
 	}
 }
 
diff --git a/Engine/Plugins/Experimental/ChaosVehiclesPlugin/Source/ChaosVehicles/Public/ChaosVehicleMovementComponent.h b/Engine/Plugins/Experimental/ChaosVehiclesPlugin/Source/ChaosVehicles/Public/ChaosVehicleMovementComponent.h
--- a/Engine/Plugins/Experimental/ChaosVehiclesPlugin/Source/ChaosVehicles/Public/ChaosVehicleMovementComponent.h
+++ b/Engine/Plugins/Experimental/ChaosVehiclesPlugin/Source/ChaosVehicles/Public/ChaosVehicleMovementComponent.h
@@ -94,7 +94,7 @@ struct FVehicleDebugParams
 	float ForceDebugScaling = 0.0006f;
 	float SleepCounterThreshold = 15;
 	bool DisableVehicleSleep = false;
-	bool EnableMultithreading = false;
+	bool EnableMultithreading = true;
 	float SetMaxMPH = 0.0f;
 };
 
