ChaosVehicles: Vehicle sleep fix; The condiftions for allowing vehicle to sleep were not correct

diff --git a/Engine/Plugins/Experimental/ChaosVehiclesPlugin/Source/ChaosVehicles/Private/ChaosVehicleMovementComponent.cpp b/Engine/Plugins/Experimental/ChaosVehiclesPlugin/Source/ChaosVehicles/Private/ChaosVehicleMovementComponent.cpp
--- a/Engine/Plugins/Experimental/ChaosVehiclesPlugin/Source/ChaosVehicles/Private/ChaosVehicleMovementComponent.cpp
+++ b/Engine/Plugins/Experimental/ChaosVehiclesPlugin/Source/ChaosVehicles/Private/ChaosVehicleMovementComponent.cpp
@@ -626,6 +626,7 @@ UChaosVehicleMovementComponent::UChaosVehicleMovementComponent(const FObjectInit
 	StabilizeControl.InitDefaults();
 
 	AngErrorAccumulator = 0.0f;
+	TargetGear = 0;
 
 	bRequiresControllerForInputs = true;
 	IdleBrakeInput = 0.0f;
@@ -923,6 +924,8 @@ void UChaosVehicleMovementComponent::SetTargetGear(int32 GearNum, bool bImmediat
 					}
 				});
 		}
+
+		TargetGear = GearNum;
 	}
 }
 
@@ -940,7 +943,7 @@ int32 UChaosVehicleMovementComponent::GetCurrentGear() const
 
 int32 UChaosVehicleMovementComponent::GetTargetGear() const
 {
-	return (PVehicleOutput)?PVehicleOutput->TargetGear:0;
+	return TargetGear;
 }
 
 bool UChaosVehicleMovementComponent::GetUseAutoGears() const
@@ -1107,6 +1110,25 @@ void UChaosVehicleMovementComponent::UpdateState(float DeltaTime)
 	// update input values
 	AController* Controller = GetController();
 	VehicleState.CaptureState(GetBodyInstance(), GetGravityZ(), DeltaTime);
+	VehicleState.NumWheelsOnGround = 0;
+	VehicleState.bVehicleInAir = false;
+	int NumWheels = 0;
+	if (PVehicleOutput)
+	{
+		for (int WheelIdx = 0; WheelIdx < PVehicleOutput->Wheels.Num(); WheelIdx++)
+		{
+			if (PVehicleOutput->Wheels[WheelIdx].InContact)
+			{
+				VehicleState.NumWheelsOnGround++;
+			}
+			else
+			{
+				VehicleState.bVehicleInAir = true;
+			}
+			NumWheels++;
+		}
+	}
+	VehicleState.bAllWheelsOnGround = (VehicleState.NumWheelsOnGround == NumWheels);
 
 	bool bProcessLocally = bRequiresControllerForInputs?(Controller && Controller->IsLocalController()):true;
 
@@ -1117,7 +1139,7 @@ void UChaosVehicleMovementComponent::UpdateState(float DeltaTime)
 		if (bReverseAsBrake)
 		{
 			//for reverse as state we want to automatically shift between reverse and first gear
-			// Note: Removed this condition to support wheel spinning when rolling backareds with accelerator pressed, rather than braking
+			// Note: Removed this condition to support wheel spinning when rolling backwards with accelerator pressed, rather than braking
 			//if (FMath::Abs(GetForwardSpeed()) < WrongDirectionThreshold)	//we only shift between reverse and first if the car is slow enough.
 			{
 				if (RawBrakeInput > KINDA_SMALL_NUMBER && GetCurrentGear() >= 0 && GetTargetGear() >= 0)
@@ -1201,9 +1223,9 @@ void UChaosVehicleMovementComponent::ProcessSleeping(const FControlInputs& Contr
 
 		const AController* Controller = GetController();
 		const bool bIsLocallyControlled = (Controller && Controller->IsLocalController());
-		const bool bControlInputPressed = bIsLocallyControlled ? (ControlInputs.ThrottleInput >= SMALL_NUMBER) || (ControlInputs.BrakeInput >= SMALL_NUMBER) || (ControlInputs.SteeringInput >= SMALL_NUMBER)
+		const bool bControlInputPressed = bIsLocallyControlled ? (ControlInputs.ThrottleInput >= SMALL_NUMBER) || (FMath::Abs(ControlInputs.SteeringInput) >= SMALL_NUMBER)
 			|| (ControlInputs.RollInput >= SMALL_NUMBER) || (ControlInputs.PitchInput >= SMALL_NUMBER) || (ControlInputs.YawInput >= SMALL_NUMBER)
-			: (ReplicatedState.ThrottleInput >= SMALL_NUMBER) || (ReplicatedState.BrakeInput >= SMALL_NUMBER) || (ReplicatedState.SteeringInput >= SMALL_NUMBER)
+			: (ReplicatedState.ThrottleInput >= SMALL_NUMBER) || (FMath::Abs(ReplicatedState.SteeringInput) >= SMALL_NUMBER)
 			|| (ReplicatedState.RollInput >= SMALL_NUMBER) || (ReplicatedState.PitchInput >= SMALL_NUMBER) || (ReplicatedState.YawInput >= SMALL_NUMBER);
 
 		// Wake if control input pressed
diff --git a/Engine/Plugins/Experimental/ChaosVehiclesPlugin/Source/ChaosVehicles/Public/ChaosVehicleMovementComponent.h b/Engine/Plugins/Experimental/ChaosVehiclesPlugin/Source/ChaosVehicles/Public/ChaosVehicleMovementComponent.h
--- a/Engine/Plugins/Experimental/ChaosVehiclesPlugin/Source/ChaosVehicles/Public/ChaosVehicleMovementComponent.h
+++ b/Engine/Plugins/Experimental/ChaosVehiclesPlugin/Source/ChaosVehicles/Public/ChaosVehicleMovementComponent.h
@@ -1331,5 +1331,6 @@ private:
 		PAerodynamicsSetup.AreaMetresSquared = Chaos::Cm2ToM2(this->DragArea);
 	}
 	Chaos::FSimpleAerodynamicsConfig PAerodynamicsSetup;
+	int32 TargetGear;
 
 };
