ChaosVehicles: Sort issues around vehicle sleeping when start awake is false.

diff --git a/Engine/Plugins/Experimental/ChaosVehiclesPlugin/Source/ChaosVehicles/Private/ChaosVehicleMovementComponent.cpp b/Engine/Plugins/Experimental/ChaosVehiclesPlugin/Source/ChaosVehicles/Private/ChaosVehicleMovementComponent.cpp
--- a/Engine/Plugins/Experimental/ChaosVehiclesPlugin/Source/ChaosVehicles/Private/ChaosVehicleMovementComponent.cpp
+++ b/Engine/Plugins/Experimental/ChaosVehiclesPlugin/Source/ChaosVehicles/Private/ChaosVehicleMovementComponent.cpp
@@ -836,10 +836,14 @@ void UChaosVehicleMovementComponent::StopMovementImmediately()
 	FBodyInstance* TargetInstance = GetBodyInstance();
 	if (TargetInstance)
 	{
-		TargetInstance->SetLinearVelocity(FVector::ZeroVector, false);
-		TargetInstance->SetAngularVelocityInRadians(FVector::ZeroVector, false);
-		TargetInstance->ClearForces();
-		TargetInstance->ClearTorques();
+		// if start awake is false then setting the velocity (even to zero) causes particle to wake up.
+		if (TargetInstance->IsInstanceAwake())
+		{
+			TargetInstance->SetLinearVelocity(FVector::ZeroVector, false);
+			TargetInstance->SetAngularVelocityInRadians(FVector::ZeroVector, false);
+			TargetInstance->ClearForces();
+			TargetInstance->ClearTorques();
+		}
 	}
 	Super::StopMovementImmediately();
 	ClearAllInput();
@@ -1229,7 +1233,7 @@ void UChaosVehicleMovementComponent::ProcessSleeping(const FControlInputs& Contr
 			|| (ReplicatedState.RollInput >= SMALL_NUMBER) || (ReplicatedState.PitchInput >= SMALL_NUMBER) || (ReplicatedState.YawInput >= SMALL_NUMBER);
 
 		// Wake if control input pressed
-		if (VehicleState.bSleeping && (bControlInputPressed || !VehicleState.bAllWheelsOnGround))
+		if (VehicleState.bSleeping && bControlInputPressed)
 		{
 			VehicleState.bSleeping = false;
 			VehicleState.SleepCounter = 0;
