Fixed PhysX build when Chaos is disabled.

diff --git a/Engine/Plugins/Experimental/ChaosVehiclesPlugin/Source/ChaosVehicles/Private/ChaosVehicleManager.cpp b/Engine/Plugins/Experimental/ChaosVehiclesPlugin/Source/ChaosVehicles/Private/ChaosVehicleManager.cpp
--- a/Engine/Plugins/Experimental/ChaosVehiclesPlugin/Source/ChaosVehicles/Private/ChaosVehicleManager.cpp
+++ b/Engine/Plugins/Experimental/ChaosVehiclesPlugin/Source/ChaosVehicles/Private/ChaosVehicleManager.cpp
@@ -23,7 +23,10 @@ uint32 FChaosVehicleManager::VehicleSetupTag = 0;
 
 FDelegateHandle FChaosVehicleManager::OnPostWorldInitializationHandle;
 FDelegateHandle FChaosVehicleManager::OnWorldCleanupHandle;
+
+#if WITH_CHAOS
 bool FChaosVehicleManager::GInitialized = false;
+#endif
 
 FChaosVehicleManager::FChaosVehicleManager(FPhysScene* PhysScene)
 #if WITH_CHAOS
@@ -35,6 +38,7 @@ FChaosVehicleManager::FChaosVehicleManager(FPhysScene* PhysScene)
 {
 	check(PhysScene);
 	
+#if WITH_CHAOS
 	if (!GInitialized)
 	{
 		GInitialized = true;
@@ -49,20 +53,24 @@ FChaosVehicleManager::FChaosVehicleManager(FPhysScene* PhysScene)
 
 	// Add to Scene-To-Manager map
 	FChaosVehicleManager::SceneToVehicleManagerMap.Add(PhysScene, this);
+#endif
 }
 
 
 void FChaosVehicleManager::RegisterCallbacks()
 {
+#if WITH_CHAOS
 	OnPhysScenePreTickHandle = Scene.OnPhysScenePreTick.AddRaw(this, &FChaosVehicleManager::Update);
 	OnPhysScenePostTickHandle = Scene.OnPhysScenePostTick.AddRaw(this, &FChaosVehicleManager::PostUpdate);
 
 	check(AsyncCallback == nullptr);
 	AsyncCallback = Scene.GetSolver()->CreateAndRegisterSimCallbackObject_External<FChaosVehicleManagerAsyncCallback>(true);
+#endif
 }
 
 void FChaosVehicleManager::UnregisterCallbacks()
 {
+#if WITH_CHAOS
 	Scene.OnPhysScenePreTick.Remove(OnPhysScenePreTickHandle);
 	Scene.OnPhysScenePostTick.Remove(OnPhysScenePostTickHandle);
 
@@ -71,6 +79,7 @@ void FChaosVehicleManager::UnregisterCallbacks()
 		Scene.GetSolver()->UnregisterAndFreeSimCallbackObject_External(AsyncCallback);
 		AsyncCallback = nullptr;
 	}
+#endif
 }
 
 void FChaosVehicleManager::OnPostWorldInitialization(UWorld* InWorld, const UWorld::InitializationValues)
@@ -163,6 +172,7 @@ void FChaosVehicleManager::ScenePreTick(FPhysScene* PhysScene, float DeltaTime)
 
 void FChaosVehicleManager::Update(FPhysScene* PhysScene, float DeltaTime)
 {
+#if WITH_CHAOS
 	SCOPE_CYCLE_COUNTER(STAT_ChaosVehicleManager_Update);
 
 	UWorld* World = Scene.GetOwningWorld();
@@ -182,6 +192,7 @@ void FChaosVehicleManager::Update(FPhysScene* PhysScene, float DeltaTime)
 			Vehicle->FinalizeSimCallbackData(*AsyncInput);
 		}
 	}
+#endif
 }
 
 void FChaosVehicleManager::PostUpdate(FChaosScene* PhysScene)
@@ -191,6 +202,7 @@ void FChaosVehicleManager::PostUpdate(FChaosScene* PhysScene)
 
 void FChaosVehicleManager::ParallelUpdateVehicles(float DeltaSeconds)
 {
+#if WITH_CHAOS
 	SCOPE_CYCLE_COUNTER(STAT_ChaosVehicleManager_ParallelUpdateVehicles);
 
 	FChaosVehicleManagerAsyncInput* AsyncInput = AsyncCallback->GetProducerInputData_External();
@@ -298,4 +310,5 @@ void FChaosVehicleManager::ParallelUpdateVehicles(float DeltaSeconds)
 
 	bool ForceSingleThread = !GVehicleDebugParams.EnableMultithreading;
 	ParallelFor(AwakeVehiclesBatch.Num(), LambdaParallelUpdate, ForceSingleThread);
+#endif
 }
diff --git a/Engine/Plugins/Experimental/ChaosVehiclesPlugin/Source/ChaosVehicles/Private/ChaosVehicleMovementComponent.cpp b/Engine/Plugins/Experimental/ChaosVehiclesPlugin/Source/ChaosVehicles/Private/ChaosVehicleMovementComponent.cpp
--- a/Engine/Plugins/Experimental/ChaosVehiclesPlugin/Source/ChaosVehicles/Private/ChaosVehicleMovementComponent.cpp
+++ b/Engine/Plugins/Experimental/ChaosVehiclesPlugin/Source/ChaosVehicles/Private/ChaosVehicleMovementComponent.cpp
@@ -1709,6 +1709,7 @@ TUniquePtr<FChaosVehicleAsyncOutput> FChaosVehicleDefaultAsyncInput::Simulate(UW
 /************************************************************************/
 void UChaosVehicleMovementComponent::Update(float DeltaTime)
 {
+#if WITH_CHAOS
 	if (CurAsyncInput)
 	{
 		if (const FBodyInstance* BodyInstance = GetBodyInstance())
@@ -1742,6 +1743,7 @@ void UChaosVehicleMovementComponent::Update(float DeltaTime)
 			}
 		}
 	}
+#endif
 }
 
 void UChaosVehicleMovementComponent::FinalizeSimCallbackData(FChaosVehicleManagerAsyncInput& Input)
diff --git a/Engine/Plugins/Experimental/ChaosVehiclesPlugin/Source/ChaosVehicles/Private/ChaosWheeledVehicleMovementComponent.cpp b/Engine/Plugins/Experimental/ChaosVehiclesPlugin/Source/ChaosVehicles/Private/ChaosWheeledVehicleMovementComponent.cpp
--- a/Engine/Plugins/Experimental/ChaosVehiclesPlugin/Source/ChaosVehicles/Private/ChaosWheeledVehicleMovementComponent.cpp
+++ b/Engine/Plugins/Experimental/ChaosVehiclesPlugin/Source/ChaosVehicles/Private/ChaosWheeledVehicleMovementComponent.cpp
@@ -1552,6 +1552,7 @@ void UChaosWheeledVehicleMovementComponent::Update(float DeltaTime)
 {
 	UChaosVehicleMovementComponent::Update(DeltaTime);
 
+#if WITH_CHAOS
 	if (CurAsyncInput)
 	{
 		if (const FBodyInstance* BodyInstance = GetBodyInstance())
@@ -1572,6 +1573,7 @@ void UChaosWheeledVehicleMovementComponent::Update(float DeltaTime)
 			}
 		}
 	}
+#endif
 }
 
 
diff --git a/Engine/Plugins/Experimental/Water/Source/Runtime/Private/BuoyancyManager.cpp b/Engine/Plugins/Experimental/Water/Source/Runtime/Private/BuoyancyManager.cpp
--- a/Engine/Plugins/Experimental/Water/Source/Runtime/Private/BuoyancyManager.cpp
+++ b/Engine/Plugins/Experimental/Water/Source/Runtime/Private/BuoyancyManager.cpp
@@ -109,7 +109,8 @@ void ABuoyancyManager::Unregister(UBuoyancyComponent* BuoyancyComponent)
 
 void ABuoyancyManager::Update(FPhysScene* PhysScene, float DeltaTime)
 {
-	if (UWorld* World = GetWorld())
+	UWorld* World = GetWorld();
+	if (World && AsyncCallback)
 	{
 		// Collect active buoyancy components
 		{
diff --git a/Engine/Plugins/Runtime/ApexDestruction/Source/ApexDestruction/Private/DestructibleComponent.cpp b/Engine/Plugins/Runtime/ApexDestruction/Source/ApexDestruction/Private/DestructibleComponent.cpp
--- a/Engine/Plugins/Runtime/ApexDestruction/Source/ApexDestruction/Private/DestructibleComponent.cpp
+++ b/Engine/Plugins/Runtime/ApexDestruction/Source/ApexDestruction/Private/DestructibleComponent.cpp
@@ -1099,7 +1099,7 @@ bool UDestructibleComponent::DoCustomNavigableGeometryExport(FNavigableGeometryE
 #if WITH_EDITORONLY_DATA
 		if (DestructibleMesh && DestructibleMesh->SourceStaticMesh)
 		{
-			GeomExport.ExportRigidBodySetup(*DestructibleMesh->SourceStaticMesh->BodySetup, GetComponentTransform());
+			GeomExport.ExportRigidBodySetup(*DestructibleMesh->SourceStaticMesh->GetBodySetup(), GetComponentTransform());
 			bExportFromBodySetup = false;
 		}
 #endif	// WITH_EDITORONLY_DATA
diff --git a/Engine/Plugins/Runtime/ApexDestruction/Source/ApexDestructionEditor/Private/SDestructibleMeshEditorViewport.cpp b/Engine/Plugins/Runtime/ApexDestruction/Source/ApexDestructionEditor/Private/SDestructibleMeshEditorViewport.cpp
--- a/Engine/Plugins/Runtime/ApexDestruction/Source/ApexDestructionEditor/Private/SDestructibleMeshEditorViewport.cpp
+++ b/Engine/Plugins/Runtime/ApexDestruction/Source/ApexDestructionEditor/Private/SDestructibleMeshEditorViewport.cpp
@@ -253,7 +253,7 @@ bool FDestructibleMeshEditorViewportClient::CanRefreshFromStaticMesh()
 		return false;
 	}
 
-	const auto* ImportInfo = DestructibleMesh->SourceStaticMesh->AssetImportData;
+	const TObjectPtr<UAssetImportData> ImportInfo = DestructibleMesh->SourceStaticMesh->AssetImportData;
 	FDateTime CurrentSourceTimestamp = FDateTime::MinValue();
 	if (ImportInfo && ImportInfo->SourceData.SourceFiles.Num() == 1)
 	{
diff --git a/Engine/Plugins/Runtime/Interchange/Source/Nodes/Public/InterchangeMaterialNode.h b/Engine/Plugins/Runtime/Interchange/Source/Nodes/Public/InterchangeMaterialNode.h
--- a/Engine/Plugins/Runtime/Interchange/Source/Nodes/Public/InterchangeMaterialNode.h
+++ b/Engine/Plugins/Runtime/Interchange/Source/Nodes/Public/InterchangeMaterialNode.h
@@ -10,6 +10,7 @@
 #include "Materials/Material.h"
 #include "Materials/MaterialInstance.h"
 #include "Materials/MaterialInterface.h"
+#include "Engine/Texture.h"
 
 #endif
 
diff --git a/Engine/Source/Runtime/Engine/Private/Collision/CollisionConversions.cpp b/Engine/Source/Runtime/Engine/Private/Collision/CollisionConversions.cpp
--- a/Engine/Source/Runtime/Engine/Private/Collision/CollisionConversions.cpp
+++ b/Engine/Source/Runtime/Engine/Private/Collision/CollisionConversions.cpp
@@ -566,8 +566,8 @@ void ConvertQueryOverlap(const FPhysicsShape& Shape, const FPhysicsActor& Actor,
 		TWeakObjectPtr<UPrimitiveComponent> OwnerComponent = CustomPayload->GetOwningComponent();
 		if (UPrimitiveComponent* OwnerComponentRaw = OwnerComponent.Get())
 		{
-			OutOverlap.Actor = OwnerComponentRaw->GetOwner();
 			OutOverlap.Component = OwnerComponent; // Copying weak pointer is faster than assigning raw pointer.
+			OutOverlap.OverlapObjectHandle = FActorInstanceHandle(OutOverlap.Component->GetOwner());
 			OutOverlap.ItemIndex = OwnerComponent->bMultiBodyOverlap ? CustomPayload->GetItemIndex() : INDEX_NONE;
 		}
 	}
diff --git a/Engine/Source/Runtime/Engine/Private/PhysicsEngine/BodyInstance.cpp b/Engine/Source/Runtime/Engine/Private/PhysicsEngine/BodyInstance.cpp
--- a/Engine/Source/Runtime/Engine/Private/PhysicsEngine/BodyInstance.cpp
+++ b/Engine/Source/Runtime/Engine/Private/PhysicsEngine/BodyInstance.cpp
@@ -1212,7 +1212,9 @@ void FInitBodiesHelperBase::CreateActor_AssumesLocked(FBodyInstance* Instance, c
 		FPhysicsInterface::CreateActor(ActorParams, Instance->ActorHandle);
 		FPhysicsInterface::SetCcdEnabled_AssumesLocked(Instance->ActorHandle, Instance->bUseCCD);
 		FPhysicsInterface::SetIsKinematic_AssumesLocked(Instance->ActorHandle, !Instance->ShouldInstanceSimulatingPhysics());
+#if WITH_CHAOS
 		FPhysicsInterface::SetMaxLinearVelocity_AssumesLocked(Instance->ActorHandle, TNumericLimits<float>::Max());
+#endif
 
 		// Set sleep even notification
 		FPhysicsInterface::SetSendsSleepNotifies_AssumesLocked(Instance->ActorHandle, Instance->bGenerateWakeEvents);
diff --git a/Engine/Source/Runtime/Engine/Private/PhysicsEngine/ConstraintInstance.cpp b/Engine/Source/Runtime/Engine/Private/PhysicsEngine/ConstraintInstance.cpp
--- a/Engine/Source/Runtime/Engine/Private/PhysicsEngine/ConstraintInstance.cpp
+++ b/Engine/Source/Runtime/Engine/Private/PhysicsEngine/ConstraintInstance.cpp
@@ -188,6 +188,7 @@ void FConstraintInstance::UpdatePlasticity()
 
 void FConstraintProfileProperties::UpdatePlasticity_AssumesLocked(const FPhysicsConstraintHandle& InConstraintRef) const
 {
+#if WITH_CHAOS
 	const float LinearPlasticityLimit = bLinearPlasticity ? LinearPlasticityThreshold : FLT_MAX;
 	const float AngularPlasticityLimit = bAngularPlasticity ? FMath::DegreesToRadians(AngularPlasticityThreshold) : MAX_FLT;
 
diff --git a/Engine/Source/Runtime/Engine/Private/PhysicsEngine/Experimental/PhysScene_Chaos.cpp b/Engine/Source/Runtime/Engine/Private/PhysicsEngine/Experimental/PhysScene_Chaos.cpp
--- a/Engine/Source/Runtime/Engine/Private/PhysicsEngine/Experimental/PhysScene_Chaos.cpp
+++ b/Engine/Source/Runtime/Engine/Private/PhysicsEngine/Experimental/PhysScene_Chaos.cpp
@@ -594,6 +594,7 @@ void FPhysScene_Chaos::RemoveObject(FSingleParticlePhysicsProxy* InObject)
 
 void FPhysScene_Chaos::RemoveObject(FGeometryCollectionPhysicsProxy* InObject)
 {
+#if WITH_CHAOS
 	Chaos::FPhysicsSolver* Solver = InObject->GetSolver<Chaos::FPhysicsSolver>();
 
 	for (TUniquePtr<Chaos::TGeometryParticle<Chaos::FReal, 3>>& GTParticleUnique : InObject->GetExternalParticles())
@@ -612,6 +613,7 @@ void FPhysScene_Chaos::RemoveObject(FGeometryCollectionPhysicsProxy* InObject)
 	}
 
 	RemoveFromComponentMaps(InObject);
+#endif
 }
 
 FPhysicsReplication* FPhysScene_Chaos::GetPhysicsReplication()
