Make visual wheel position more accurate. This change also helps fix wheels to correct height when vehicle transitions from crowd vehicle to simulated vehicle.

diff --git a/Engine/Plugins/Experimental/ChaosVehiclesPlugin/Source/ChaosVehicles/Private/ChaosVehicleMovementComponent.cpp b/Engine/Plugins/Experimental/ChaosVehiclesPlugin/Source/ChaosVehicles/Private/ChaosVehicleMovementComponent.cpp
--- a/Engine/Plugins/Experimental/ChaosVehiclesPlugin/Source/ChaosVehicles/Private/ChaosVehicleMovementComponent.cpp
+++ b/Engine/Plugins/Experimental/ChaosVehiclesPlugin/Source/ChaosVehicles/Private/ChaosVehicleMovementComponent.cpp
@@ -1391,6 +1391,9 @@ void UChaosVehicleMovementComponent::CreateVehicle()
 				VehicleSimulationPT->Init(PVehicle);
 
 			}
+
+			VehicleState.CaptureState(GetBodyInstance(), GetGravityZ(), 0.01667f);
+
 		}
 	}
 }
diff --git a/Engine/Plugins/Experimental/ChaosVehiclesPlugin/Source/ChaosVehicles/Private/ChaosVehicleWheel.cpp b/Engine/Plugins/Experimental/ChaosVehiclesPlugin/Source/ChaosVehicles/Private/ChaosVehicleWheel.cpp
--- a/Engine/Plugins/Experimental/ChaosVehiclesPlugin/Source/ChaosVehicles/Private/ChaosVehicleWheel.cpp
+++ b/Engine/Plugins/Experimental/ChaosVehiclesPlugin/Source/ChaosVehicles/Private/ChaosVehicleWheel.cpp
@@ -87,7 +87,7 @@ float UChaosVehicleWheel::GetWheelAngularVelocity() const
 float UChaosVehicleWheel::GetSuspensionOffset() const
 {
 	check(VehicleComponent && VehicleComponent->PhysicsVehicleOutput());
-	return VehicleComponent->PhysicsVehicleOutput()->Wheels[WheelIndex].SuspensionOffset;
+	return VehicleComponent->GetSuspensionOffset(WheelIndex);
 }
 
 bool UChaosVehicleWheel::IsInAir() const
diff --git a/Engine/Plugins/Experimental/ChaosVehiclesPlugin/Source/ChaosVehicles/Private/ChaosWheeledVehicleMovementComponent.cpp b/Engine/Plugins/Experimental/ChaosVehiclesPlugin/Source/ChaosVehicles/Private/ChaosWheeledVehicleMovementComponent.cpp
--- a/Engine/Plugins/Experimental/ChaosVehiclesPlugin/Source/ChaosVehicles/Private/ChaosWheeledVehicleMovementComponent.cpp
+++ b/Engine/Plugins/Experimental/ChaosVehiclesPlugin/Source/ChaosVehicles/Private/ChaosWheeledVehicleMovementComponent.cpp
@@ -2023,34 +2023,41 @@ void UChaosWheeledVehicleMovementComponent::DrawDial(UCanvas* Canvas, FVector2D
 
 void UChaosWheeledVehicleMovementComponent::FillWheelOutputState()
 {
-	for (int WheelIdx = 0; WheelIdx < Wheels.Num(); WheelIdx++)
+	if (const FChaosVehicleAsyncOutput* CurrentOutput = static_cast<FChaosVehicleAsyncOutput*>(CurAsyncOutput))
 	{
-		auto& PWheel = PVehicleOutput->Wheels[WheelIdx];
-		FHitResult& HitResult = Wheels[WheelIdx]->HitResult;
-
-		FWheelStatus& State = WheelStatus[WheelIdx];
-
-		State.bInContact = HitResult.bBlockingHit;
-		State.ContactPoint = HitResult.ImpactPoint;
-		State.PhysMaterial = HitResult.PhysMaterial;
-		State.NormalizedSuspensionLength = PWheel.NormalizedSuspensionLength;
-		State.SpringForce = PWheel.SpringForce;
-		State.SlipAngle = PWheel.SlipAngle;
-		State.bIsSlipping = PWheel.bIsSlipping;
-		State.SlipMagnitude = PWheel.SlipMagnitude;
-		State.bIsSkidding = PWheel.bIsSkidding;
-		State.SkidMagnitude = PWheel.SkidMagnitude;
-		if (State.bIsSkidding)
-		{
-			State.SkidNormal = PWheel.SkidNormal;
-			//DrawDebugLine(GetWorld()
-			//	, State.ContactPoint
-			//	, State.ContactPoint + State.SkidNormal
-			//	, FColor::Yellow, true, -1.0f, 0, 4);
-		}
-		else
+		if (CurrentOutput->bValid && PVehicleOutput)
 		{
-			State.SkidNormal = FVector::ZeroVector;
+			for (int WheelIdx = 0; WheelIdx < Wheels.Num(); WheelIdx++)
+			{
+				auto& PWheel = PVehicleOutput->Wheels[WheelIdx];
+				FHitResult& HitResult = Wheels[WheelIdx]->HitResult;
+
+				FWheelStatus& State = WheelStatus[WheelIdx];
+
+				State.bIsValid = true;
+				State.bInContact = HitResult.bBlockingHit;
+				State.ContactPoint = HitResult.ImpactPoint;
+				State.PhysMaterial = HitResult.PhysMaterial;
+				State.NormalizedSuspensionLength = PWheel.NormalizedSuspensionLength;
+				State.SpringForce = PWheel.SpringForce;
+				State.SlipAngle = PWheel.SlipAngle;
+				State.bIsSlipping = PWheel.bIsSlipping;
+				State.SlipMagnitude = PWheel.SlipMagnitude;
+				State.bIsSkidding = PWheel.bIsSkidding;
+				State.SkidMagnitude = PWheel.SkidMagnitude;
+				if (State.bIsSkidding)
+				{
+					State.SkidNormal = PWheel.SkidNormal;
+					//DrawDebugLine(GetWorld()
+					//	, State.ContactPoint
+					//	, State.ContactPoint + State.SkidNormal
+					//	, FColor::Yellow, true, -1.0f, 0, 4);
+				}
+				else
+				{
+					State.SkidNormal = FVector::ZeroVector;
+				}
+			}
 		}
 	}
 }
@@ -2493,6 +2500,73 @@ void UChaosWheeledVehicleMovementComponent::SetWheelMaxSteerAngle(int WheelIndex
 }
 
 
+float UChaosWheeledVehicleMovementComponent::GetSuspensionOffset(int WheelIndex)
+{
+	float Offset = 0.f;
+
+	FChaosWheelSetup& WheelSetup = WheelSetups[WheelIndex];
+	FTransform VehicleWorldTransform = GetBodyInstance()->GetUnrealWorldTransform();
+	if (GetBodyInstance())
+	{
+		if (UChaosVehicleWheel* Wheel = WheelSetups[WheelIndex].WheelClass.GetDefaultObject())
+		{
+			if (WheelStatus[WheelIndex].bIsValid)
+			{
+				if (WheelStatus[WheelIndex].bInContact)
+				{
+					FVector LocalPos = GetWheelRestingPosition(WheelSetup);
+					FVector LocalHitPoint = VehicleWorldTransform.InverseTransformPosition(WheelStatus[WheelIndex].ContactPoint);
+					Offset = LocalHitPoint.Z - LocalPos.Z + PVehicleOutput->Wheels[WheelIndex].WheelRadius;
+					Offset = FMath::Clamp(Offset, -Wheel->SuspensionMaxDrop, Wheel->SuspensionMaxRaise);
+				}
+				else
+				{
+					Offset = -Wheel->SuspensionMaxDrop;
+				}
+			}
+			else
+			{
+				ECollisionChannel SpringCollisionChannel = ECollisionChannel::ECC_WorldDynamic;
+				FCollisionResponseParams ResponseParams;
+				ResponseParams.CollisionResponse = WheelTraceCollisionResponses;
+
+				TArray<AActor*> ActorsToIgnore;
+				ActorsToIgnore.Add(GetPawnOwner()); // ignore self in scene query
+
+				FCollisionQueryParams TraceParams(NAME_None, FCollisionQueryParams::GetUnknownStatId(), false, nullptr);
+				TraceParams.bReturnPhysicalMaterial = true;	// we need this to get the surface friction coefficient
+				TraceParams.AddIgnoredActors(ActorsToIgnore);
+				TraceParams.bTraceComplex = (Wheels[WheelIndex]->SweepType == ESweepType::ComplexSweep);
+
+				FVector LocalDirection = Wheel->SuspensionAxis;
+				FVector WorldLocation = VehicleWorldTransform.TransformPosition(GetWheelRestingPosition(WheelSetup));
+				FVector WorldDirection = VehicleWorldTransform.TransformVector(LocalDirection);
+
+				FVector TraceStart = WorldLocation - WorldDirection * (Wheel->SuspensionMaxRaise);
+				FVector TraceEnd = WorldLocation + WorldDirection * (Wheel->SuspensionMaxDrop + Wheel->WheelRadius);
+
+				FHitResult HitResult;
+				GetWorld()->LineTraceSingleByChannel(HitResult, TraceStart, TraceEnd, SpringCollisionChannel, TraceParams, ResponseParams);
+
+				if (HitResult.bBlockingHit)
+				{
+					FVector LocalPos = GetWheelRestingPosition(WheelSetup);
+					FVector LocalHitPoint = VehicleWorldTransform.InverseTransformPosition(HitResult.ImpactPoint);
+					Offset = LocalHitPoint.Z - LocalPos.Z + Wheel->WheelRadius;
+					Offset = FMath::Clamp(Offset, -Wheel->SuspensionMaxDrop, Wheel->SuspensionMaxRaise);
+				}
+				else
+				{
+					Offset = -Wheel->SuspensionMaxDrop;
+				}
+
+			}
+		}
+	}
+
+	return Offset;
+}
+
 FChaosWheelSetup::FChaosWheelSetup()
 	: WheelClass(UChaosVehicleWheel::StaticClass())
 //	, SteeringBoneName(NAME_None)
diff --git a/Engine/Plugins/Experimental/ChaosVehiclesPlugin/Source/ChaosVehicles/Public/ChaosVehicleMovementComponent.h b/Engine/Plugins/Experimental/ChaosVehiclesPlugin/Source/ChaosVehicles/Public/ChaosVehicleMovementComponent.h
--- a/Engine/Plugins/Experimental/ChaosVehiclesPlugin/Source/ChaosVehicles/Public/ChaosVehicleMovementComponent.h
+++ b/Engine/Plugins/Experimental/ChaosVehiclesPlugin/Source/ChaosVehicles/Public/ChaosVehicleMovementComponent.h
@@ -1057,6 +1057,8 @@ public:
 		return PVehicleOutput;
 	}
 
+	virtual float GetSuspensionOffset(int WheelIndex) { return 0.f; }
+
 	//----ASYNC----
 	TUniquePtr<FChaosVehicleAsyncInput> SetCurrentAsyncInputOutput(int32 InputIdx, FChaosVehicleManagerAsyncOutput* CurOutput, FChaosVehicleManagerAsyncOutput* NextOutput, float Alpha, int32 VehicleManagerTimestamp);
 
diff --git a/Engine/Plugins/Experimental/ChaosVehiclesPlugin/Source/ChaosVehicles/Public/ChaosWheeledVehicleMovementComponent.h b/Engine/Plugins/Experimental/ChaosVehiclesPlugin/Source/ChaosVehicles/Public/ChaosWheeledVehicleMovementComponent.h
--- a/Engine/Plugins/Experimental/ChaosVehiclesPlugin/Source/ChaosVehicles/Public/ChaosWheeledVehicleMovementComponent.h
+++ b/Engine/Plugins/Experimental/ChaosVehiclesPlugin/Source/ChaosVehicles/Public/ChaosWheeledVehicleMovementComponent.h
@@ -134,6 +134,7 @@ struct CHAOSVEHICLES_API FWheelStatus
 
 	explicit FWheelStatus(ENoInit NoInit)
 	{
+		bIsValid = false;
 	}
 
 	void Init()
@@ -148,10 +149,12 @@ struct CHAOSVEHICLES_API FWheelStatus
 		SpringForce = 0.f;
 		SkidNormal = FVector::ZeroVector;
 		ContactPoint = FVector::ZeroVector;
+		bIsValid = false;
 	}
 
 	FString ToString() const;
 
+	bool bIsValid;
 };
 
 USTRUCT()
@@ -702,6 +705,7 @@ class CHAOSVEHICLES_API UChaosWheeledVehicleMovementComponent : public UChaosVeh
 		return WheelStatus[WheelIndex];
 	}
 
+	virtual float GetSuspensionOffset(int WheelIndex) override;
 
 	/** Set all channels to the specified response - for wheel raycasts */
 	void SetWheelTraceAllChannels(ECollisionResponse NewResponse)
