ChaosVehicles: Make some functions public

diff --git a/Engine/Plugins/Experimental/ChaosVehiclesPlugin/Source/ChaosVehicles/Public/ChaosWheeledVehicleMovementComponent.h b/Engine/Plugins/Experimental/ChaosVehiclesPlugin/Source/ChaosVehicles/Public/ChaosWheeledVehicleMovementComponent.h
--- a/Engine/Plugins/Experimental/ChaosVehiclesPlugin/Source/ChaosVehicles/Public/ChaosWheeledVehicleMovementComponent.h
+++ b/Engine/Plugins/Experimental/ChaosVehiclesPlugin/Source/ChaosVehicles/Public/ChaosWheeledVehicleMovementComponent.h
@@ -831,11 +831,6 @@ class CHAOSVEHICLES_API UChaosWheeledVehicleMovementComponent : public UChaosVeh
 	UFUNCTION(BlueprintCallable, Category = "Game|Components|ChaosWheeledVehicleMovement")
 	void SetWheelMaxSteerAngle(int WheelIndex, float AngleDegrees);
 
-protected:
-
-	//////////////////////////////////////////////////////////////////////////
-	// Setup
-
 	/** */
 	virtual TUniquePtr<Chaos::FSimpleWheeledVehicle> CreatePhysicsVehicle() override
 	{
@@ -845,15 +840,20 @@ protected:
 		return UChaosVehicleMovementComponent::CreatePhysicsVehicle();
 	}
 
+	/** Allocate and setup the Chaos vehicle */
+	virtual void SetupVehicle(TUniquePtr<Chaos::FSimpleWheeledVehicle>& PVehicle) override;
+
+protected:
+
+	//////////////////////////////////////////////////////////////////////////
+	// Setup
+
 	/** Re-Compute any runtime constants values that rely on setup data */
 	virtual void ComputeConstants() override;
 
 	/** Skeletal mesh needs some special handling in the vehicle case */
 	virtual void FixupSkeletalMesh();
 
-	/** Allocate and setup the Chaos vehicle */
-	virtual void SetupVehicle(TUniquePtr<Chaos::FSimpleWheeledVehicle>& PVehicle) override;
-
 	/** Create and setup the Chaos vehicle */
 	virtual void CreateVehicle();
 
