Delete already commented out lines of code

diff --git a/Engine/Plugins/Experimental/ChaosVehiclesPlugin/Source/ChaosVehicles/Public/ChaosVehicleWheel.h b/Engine/Plugins/Experimental/ChaosVehiclesPlugin/Source/ChaosVehicles/Public/ChaosVehicleWheel.h
--- a/Engine/Plugins/Experimental/ChaosVehiclesPlugin/Source/ChaosVehicles/Public/ChaosVehicleWheel.h
+++ b/Engine/Plugins/Experimental/ChaosVehiclesPlugin/Source/ChaosVehicles/Public/ChaosVehicleWheel.h
@@ -64,17 +64,6 @@ class UChaosWheeledVehicleMovementComponent;
 		UPROPERTY(EditDefaultsOnly, Category = Shape)
 		class UStaticMesh* CollisionMesh;
 
-		/** If set, shape won't be created, but mapped from chassis mesh */
-		//UPROPERTY(EditDefaultsOnly, Category = Shape, meta = (DisplayName = "UsePhysAssetShape"))
-		//bool bDontCreateShape;
-
-		/**
-		 *	If true, ShapeRadius and ShapeWidth will be used to automatically scale collision taken from CollisionMesh to match wheel size.
-		 *	If false, size of CollisionMesh won't be changed. Use if you want to scale wheels manually.
-		 */
-		//UPROPERTY(EditAnywhere, Category = Shape)
-		//bool bAutoAdjustCollisionSize;
-
 		/** If left undefined then the bAffectedByEngine value is used, if defined then bAffectedByEngine is ignored and the differential setup on the vehicle defines which wheels get power from the engine */
 		UPROPERTY(EditAnywhere, Category = Wheel)
 		EAxleType AxleType;
@@ -94,10 +83,6 @@ class UChaosWheeledVehicleMovementComponent;
 		UPROPERTY(EditAnywhere, Category = Wheel, meta = (ClampMin = "0.01", UIMin = "0.01"))
 		float WheelWidth;
 
-		/** Mass of this wheel */
-		//UPROPERTY(EditAnywhere, Category = Wheel, meta = (ClampMin = "0.01", UIMin = "0.01"))
-		//float WheelMass;
-
 		/** Tyre Cornering Ability */
 		UPROPERTY(EditAnywhere, Category = Wheel, meta = (EditCondition = "bNewWheelSimulation"))
 		float CorneringStiffness;
@@ -118,10 +103,6 @@ class UChaosWheeledVehicleMovementComponent;
 		UPROPERTY(EditAnywhere, Category = Wheel, meta = (ClampMin = "0.0", UIMin = "0.0"))
 		float SkidThreshold;
 
-		///** Damping rate for this wheel (Kgm^2/s) */
-		//UPROPERTY(EditAnywhere, Category = Wheel, meta = (ClampMin = "0.01", UIMin = "0.01"))
-		//float DampingRate;
-
 		// steer angle in degrees for this wheel
 		UPROPERTY(EditAnywhere, Category = WheelsSetup)
 		float MaxSteerAngle;
@@ -153,18 +134,6 @@ class UChaosWheeledVehicleMovementComponent;
 		UPROPERTY(EditAnywhere, Category = Setup)
 		FRuntimeFloatCurve LateralSlipGraph;
 
-		///** Max normalized tire load at which the tire can deliver no more lateral stiffness no matter how much extra load is applied to the tire. */
-		//UPROPERTY(EditAnywhere, Category = Tire, meta = (ClampMin = "0.01", UIMin = "0.01"))
-		//float LatStiffMaxLoad;
-
-		///** How much lateral stiffness to have given lateral slip */
-		//UPROPERTY(EditAnywhere, Category = Tire, meta = (ClampMin = "0.01", UIMin = "0.01"))
-		//float LatStiffValue;
-
-		///** How much longitudinal stiffness to have given longitudinal slip */
-		//UPROPERTY(EditAnywhere, Category = Tire)
-		//float LongStiffValue;
-
 		/** Local body direction in which where suspension forces are applied (typically along -Z-axis) */
 		UPROPERTY(EditAnywhere, Category = Suspension)
 		FVector SuspensionAxis;
@@ -181,10 +150,7 @@ class UChaosWheeledVehicleMovementComponent;
 		UPROPERTY(EditAnywhere, Category = Suspension)
 		float SuspensionMaxDrop;
 
-		///** Oscillation frequency of suspension. Standard cars have values between 5 and 10 */
-		//UPROPERTY(EditAnywhere, Category = Suspension)
-		//float SuspensionNaturalFrequency;
-
+		/** Suspension damping, larger value causes the suspension to come to rest faster [range 0 to 1] */
 		UPROPERTY(EditAnywhere, Category = Suspension)
 		float SuspensionDampingRatio;
 
@@ -210,15 +176,6 @@ class UChaosWheeledVehicleMovementComponent;
 		UPROPERTY(EditAnywhere, Category = Suspension, meta = (UIMin = "0", UIMax = "1"))
 		float RollbarScaling;
 
-		// Now calculating damping from Suspension Damping Ratio - normalized value is slightly more meaningful
-		///** Dampen rate of change of spring compression */
-		//UPROPERTY(EditAnywhere, Category = Suspension)
-		//float CompressionDamping;
-
-		///** Dampen rate of change of spring extension */
-		//UPROPERTY(EditAnywhere, Category = Suspension)
-		//float ReboundDamping;
-
 		/** Whether wheel suspension considers simple, complex, or both */
 		UPROPERTY(EditAnywhere, Category = Suspension)
 		ESweepShape SweepShape;
diff --git a/Engine/Plugins/Experimental/ChaosVehiclesPlugin/Source/ChaosVehicles/Public/ChaosWheeledVehicleMovementComponent.h b/Engine/Plugins/Experimental/ChaosVehiclesPlugin/Source/ChaosVehicles/Public/ChaosWheeledVehicleMovementComponent.h
--- a/Engine/Plugins/Experimental/ChaosVehiclesPlugin/Source/ChaosVehicles/Public/ChaosWheeledVehicleMovementComponent.h
+++ b/Engine/Plugins/Experimental/ChaosVehiclesPlugin/Source/ChaosVehicles/Public/ChaosWheeledVehicleMovementComponent.h
@@ -333,11 +333,6 @@ struct FVehicleTransmissionConfig
 	UPROPERTY(EditAnywhere, AdvancedDisplay, Category = Setup)
 	float TransmissionEfficiency;
 
-	/** Value of engineRevs/maxEngineRevs that is high enough to increment gear*/
-	//UPROPERTY(EditAnywhere, AdvancedDisplay, Category = Setup, meta = (ClampMin = "0.0", UIMin = "0.0", ClampMax = "1.0", UIMax = "1.0"))
-	//float NeutralGearUpRatio;
-
-
 	const Chaos::FSimpleTransmissionConfig& GetPhysicsTransmissionConfig()
 	{
 		FillTransmissionSetup();
