Add HitMaterial and HitComponent to vehicle debug CVars

diff --git a/Engine/Plugins/Experimental/ChaosVehiclesPlugin/Source/ChaosVehicles/Private/ChaosWheeledVehicleMovementComponent.cpp b/Engine/Plugins/Experimental/ChaosVehiclesPlugin/Source/ChaosVehicles/Private/ChaosWheeledVehicleMovementComponent.cpp
--- a/Engine/Plugins/Experimental/ChaosVehiclesPlugin/Source/ChaosVehicles/Private/ChaosWheeledVehicleMovementComponent.cpp
+++ b/Engine/Plugins/Experimental/ChaosVehiclesPlugin/Source/ChaosVehicles/Private/ChaosWheeledVehicleMovementComponent.cpp
@@ -53,6 +53,8 @@ FAutoConsoleVariableRef CVarChaosVehiclesShowSuspensionLimits(TEXT("p.Vehicle.Sh
 FAutoConsoleVariableRef CVarChaosVehiclesShowWheelForces(TEXT("p.Vehicle.ShowWheelForces"), GWheeledVehicleDebugParams.ShowWheelForces, TEXT("Enable/Disable Wheel Forces Visualisation."));
 FAutoConsoleVariableRef CVarChaosVehiclesShowSuspensionForces(TEXT("p.Vehicle.ShowSuspensionForces"), GWheeledVehicleDebugParams.ShowSuspensionForces, TEXT("Enable/Disable Suspension Forces Visualisation."));
 FAutoConsoleVariableRef CVarChaosVehiclesShowBatchQueryExtents(TEXT("p.Vehicle.ShowBatchQueryExtents"), GWheeledVehicleDebugParams.ShowBatchQueryExtents, TEXT("Enable/Disable Suspension Forces Visualisation."));
+FAutoConsoleVariableRef CVarChaosVehiclesShowRaycastComponent(TEXT("p.Vehicle.ShowRaycastComponent"), GWheeledVehicleDebugParams.ShowRaycastComponent, TEXT("Enable/Disable Raycast Component Hit Visualisation."));
+FAutoConsoleVariableRef CVarChaosVehiclesShowRaycastMaterial(TEXT("p.Vehicle.ShowRaycastMaterial"), GWheeledVehicleDebugParams.ShowRaycastMaterial, TEXT("Enable/Disable Raycast Material Hit Visualisation."));
 FAutoConsoleVariableRef CVarChaosVehiclesTraceTypeOverride(TEXT("p.Vehicle.TraceTypeOverride"), GWheeledVehicleDebugParams.TraceTypeOverride, TEXT("Override ray trace type, 1=Simple, 2=Complex."));
 
 FAutoConsoleVariableRef CVarChaosVehiclesDisableSuspensionForces(TEXT("p.Vehicle.DisableSuspensionForces"), GWheeledVehicleDebugParams.DisableSuspensionForces, TEXT("Enable/Disable Suspension Forces."));
@@ -897,14 +899,14 @@ void UChaosWheeledVehicleSimulation::DrawDebug3D()
 		}
 	}
 
-	if (GWheeledVehicleDebugParams.ShowWheelCollisionNormal)
+	if (GWheeledVehicleDebugParams.ShowRaycastComponent)
 	{
-		FString Name;
-		for (int WheelIdx = 0; WheelIdx < Wheels.Num(); WheelIdx++)
+		for (int WheelIdx = 0; WheelIdx < PVehicle->Suspension.Num(); WheelIdx++)
 		{
 			FHitResult& Hit = Wheels[WheelIdx]->HitResult;
 
 			FVector VehicleRightAxis = VehicleState.VehicleWorldTransform.GetUnitAxis(EAxis::Y) * 20.0f;
+			FVector VehicleUpAxis = VehicleState.VehicleWorldTransform.GetUnitAxis(EAxis::Z) * 20.0f;
 			const FVector& WheelOffset = PVehicle->Suspension[WheelIdx].GetLocalRestingPosition();
 			if (WheelOffset.Y < 0.0f)
 			{
@@ -912,19 +914,53 @@ void UChaosWheeledVehicleSimulation::DrawDebug3D()
 			}
 
 			FVector Pt = Hit.ImpactPoint + VehicleRightAxis;
-			FDebugDrawQueue::GetInstance().DrawDebugLine(Pt, Pt + Hit.Normal * 20.0f, FColor::Yellow, false, 1.0f, 0, 1.0f);
-			FDebugDrawQueue::GetInstance().DrawDebugSphere(Pt, 5.0f, 4, FColor::White, false, 1.0f, 0, 1.0f);
 			if (Hit.GetComponent())
 			{
 				FDebugDrawQueue::GetInstance().DrawDebugString(Pt + VehicleRightAxis, Hit.GetComponent()->GetName(), nullptr, FColor::White, -1.f, true, 1.0f);
+			}
+		}
+	}
 
-				AActor* OwningActor = Hit.GetComponent()->GetOwner();
-				if (OwningActor)
-				{
-					Name = OwningActor->GetName();
-				}
+	if (GWheeledVehicleDebugParams.ShowRaycastMaterial)
+	{
+		for (int WheelIdx = 0; WheelIdx < PVehicle->Suspension.Num(); WheelIdx++)
+		{
+			FHitResult& Hit = Wheels[WheelIdx]->HitResult;
 
+			FVector VehicleRightAxis = VehicleState.VehicleWorldTransform.GetUnitAxis(EAxis::Y) * 20.0f;
+			FVector VehicleUpAxis = VehicleState.VehicleWorldTransform.GetUnitAxis(EAxis::Z) * 20.0f;
+			const FVector& WheelOffset = PVehicle->Suspension[WheelIdx].GetLocalRestingPosition();
+			if (WheelOffset.Y < 0.0f)
+			{
+				VehicleRightAxis = VehicleRightAxis * -1.0f;
 			}
+
+			FVector Pt = Hit.ImpactPoint + VehicleRightAxis;
+			if (Hit.PhysMaterial.IsValid())
+			{
+				FDebugDrawQueue::GetInstance().DrawDebugString(Pt + VehicleRightAxis + VehicleUpAxis, Hit.PhysMaterial->GetName(), nullptr, FColor::White, -1.f, true, 1.0f);
+			}
+		}
+
+	}
+
+	if (GWheeledVehicleDebugParams.ShowWheelCollisionNormal)
+	{
+		FString Name;
+		for (int WheelIdx = 0; WheelIdx < Wheels.Num(); WheelIdx++)
+		{
+			FHitResult& Hit = Wheels[WheelIdx]->HitResult;
+
+			FVector VehicleRightAxis = VehicleState.VehicleWorldTransform.GetUnitAxis(EAxis::Y) * 20.0f;
+			const FVector& WheelOffset = PVehicle->Suspension[WheelIdx].GetLocalRestingPosition();
+			if (WheelOffset.Y < 0.0f)
+			{
+				VehicleRightAxis = VehicleRightAxis * -1.0f;
+			}
+
+			FVector Pt = Hit.ImpactPoint + VehicleRightAxis;
+			FDebugDrawQueue::GetInstance().DrawDebugLine(Pt, Pt + Hit.Normal * 20.0f, FColor::Yellow, false, 1.0f, 0, 1.0f);
+			FDebugDrawQueue::GetInstance().DrawDebugSphere(Pt, 5.0f, 4, FColor::White, false, 1.0f, 0, 1.0f);
 		}
 	}
 
diff --git a/Engine/Plugins/Experimental/ChaosVehiclesPlugin/Source/ChaosVehicles/Public/ChaosWheeledVehicleMovementComponent.h b/Engine/Plugins/Experimental/ChaosVehiclesPlugin/Source/ChaosVehicles/Public/ChaosWheeledVehicleMovementComponent.h
--- a/Engine/Plugins/Experimental/ChaosVehiclesPlugin/Source/ChaosVehicles/Public/ChaosWheeledVehicleMovementComponent.h
+++ b/Engine/Plugins/Experimental/ChaosVehiclesPlugin/Source/ChaosVehicles/Public/ChaosWheeledVehicleMovementComponent.h
@@ -24,6 +24,8 @@ struct FWheeledVehicleDebugParams
 	bool ShowWheelForces = false;
 	bool ShowSuspensionForces = false;
 	bool ShowBatchQueryExtents = false;
+	bool ShowRaycastComponent = false;
+	bool ShowRaycastMaterial = false;
 	int TraceTypeOverride = 0;
 
 	bool DisableSuspensionForces = false;
