Interface to force vehicles to sleep

#obforgithub public, internal

diff --git a/Engine/Plugins/Experimental/ChaosVehiclesPlugin/Source/ChaosVehicles/Private/ChaosVehicleMovementComponent.cpp b/Engine/Plugins/Experimental/ChaosVehiclesPlugin/Source/ChaosVehicles/Private/ChaosVehicleMovementComponent.cpp
--- a/Engine/Plugins/Experimental/ChaosVehiclesPlugin/Source/ChaosVehicles/Private/ChaosVehicleMovementComponent.cpp
+++ b/Engine/Plugins/Experimental/ChaosVehiclesPlugin/Source/ChaosVehicles/Private/ChaosVehicleMovementComponent.cpp
@@ -902,6 +902,20 @@ void UChaosVehicleMovementComponent::SetParked(bool bParked)
 	bParkEnabled = bParked;
 }
 
+void UChaosVehicleMovementComponent::SetSleeping(bool bEnableSleep)
+{
+	if (bEnableSleep)
+	{
+		PutAllEnabledRigidBodiesToSleep();
+		VehicleState.bSleeping = true;
+	}
+	else
+	{
+		WakeAllEnabledRigidBodies();
+		VehicleState.bSleeping = false;
+	}
+}
+
 void UChaosVehicleMovementComponent::SetChangeUpInput(bool bNewGearUp)
 {
 	bRawGearUpInput = bNewGearUp;
@@ -1237,7 +1251,7 @@ void UChaosVehicleMovementComponent::ProcessSleeping(const FControlInputs& Contr
 		{
 			VehicleState.bSleeping = false;
 			VehicleState.SleepCounter = 0;
-			TargetInstance->WakeInstance();
+			SetSleeping(false);
 		}
 		else if (!VehicleState.bSleeping && !bControlInputPressed && VehicleState.bAllWheelsOnGround && (VehicleState.VehicleUpAxis.Z > SleepSlopeLimit))
 		{
@@ -1251,7 +1265,7 @@ void UChaosVehicleMovementComponent::ProcessSleeping(const FControlInputs& Contr
 				else
 				{
 					VehicleState.bSleeping = true;
-					TargetInstance->PutInstanceToSleep();
+					SetSleeping(true);
 				}
 			}
 		}
@@ -1870,6 +1884,41 @@ void UChaosVehicleMovementComponent::SetBaseSnapshot(const FBaseSnapshotData& Sn
 	}
 }
 
+void UChaosVehicleMovementComponent::WakeAllEnabledRigidBodies()
+{
+	if (USkeletalMeshComponent* Mesh = GetSkeletalMesh())
+	{
+		for (int32 i = 0; i < Mesh->Bodies.Num(); i++)
+		{
+			FBodyInstance* BI = Mesh->Bodies[i];
+			check(BI);
+
+			if (!BI->IsPhysicsDisabled() && BI->IsNonKinematic())
+			{
+				BI->WakeInstance();
+			}
+		}
+	}
+}
+
+void UChaosVehicleMovementComponent::PutAllEnabledRigidBodiesToSleep()
+{
+	if (USkeletalMeshComponent* Mesh = GetSkeletalMesh())
+	{
+		for (int32 i = 0; i < Mesh->Bodies.Num(); i++)
+		{
+			FBodyInstance* BI = Mesh->Bodies[i];
+			check(BI);
+
+			if (!BI->IsPhysicsDisabled() && BI->IsNonKinematic())
+			{
+				BI->PutInstanceToSleep();
+			}
+		}
+	}
+}
+
+
 #undef LOCTEXT_NAMESPACE
 
 
diff --git a/Engine/Plugins/Experimental/ChaosVehiclesPlugin/Source/ChaosVehicles/Public/ChaosVehicleMovementComponent.h b/Engine/Plugins/Experimental/ChaosVehiclesPlugin/Source/ChaosVehicles/Public/ChaosVehicleMovementComponent.h
--- a/Engine/Plugins/Experimental/ChaosVehiclesPlugin/Source/ChaosVehicles/Public/ChaosVehicleMovementComponent.h
+++ b/Engine/Plugins/Experimental/ChaosVehiclesPlugin/Source/ChaosVehicles/Public/ChaosVehicleMovementComponent.h
@@ -978,6 +978,10 @@ public:
 	UFUNCTION(BlueprintCallable, Category = "Game|Components|ChaosVehicleMovement")
 	void SetParked(bool bParked);
 
+	/** Set the vehicle sleeping (bEnableSleep=true) or wake it up (bEnableSleep=false) */
+	UFUNCTION(BlueprintCallable, Category = "Game|Components|ChaosVehicleMovement")
+	void SetSleeping(bool bEnableSleep);
+
 	/** Set the user input for gear up */
 	UFUNCTION(BlueprintCallable, Category="Game|Components|ChaosVehicleMovement")
 	void SetChangeUpInput(bool bNewGearUp);
@@ -1332,6 +1336,10 @@ private:
 		PAerodynamicsSetup.DownforceCoefficient = this->DownforceCoefficient;
 		PAerodynamicsSetup.AreaMetresSquared = Chaos::Cm2ToM2(this->DragArea);
 	}
+
+	void WakeAllEnabledRigidBodies();
+	void PutAllEnabledRigidBodiesToSleep();
+
 	Chaos::FSimpleAerodynamicsConfig PAerodynamicsSetup;
 	int32 TargetGear;
 
