Vehicle Sleep disable option still requires vehicle wake from control inputs

diff --git a/Engine/Plugins/Experimental/ChaosVehiclesPlugin/Source/ChaosVehicles/Private/ChaosVehicleMovementComponent.cpp b/Engine/Plugins/Experimental/ChaosVehiclesPlugin/Source/ChaosVehicles/Private/ChaosVehicleMovementComponent.cpp
--- a/Engine/Plugins/Experimental/ChaosVehiclesPlugin/Source/ChaosVehicles/Private/ChaosVehicleMovementComponent.cpp
+++ b/Engine/Plugins/Experimental/ChaosVehiclesPlugin/Source/ChaosVehicles/Private/ChaosVehicleMovementComponent.cpp
@@ -809,7 +809,6 @@ void UChaosVehicleMovementComponent::PreTickGT(float DeltaTime)
 		}
 	}
 
-	if (!GVehicleDebugParams.DisableVehicleSleep)
 	{
 		// is this needless copying
 		FControlInputs ControlInputs;
@@ -1253,7 +1252,7 @@ void UChaosVehicleMovementComponent::ProcessSleeping(const FControlInputs& Contr
 			VehicleState.SleepCounter = 0;
 			SetSleeping(false);
 		}
-		else if (!VehicleState.bSleeping && !bControlInputPressed && VehicleState.bAllWheelsOnGround && (VehicleState.VehicleUpAxis.Z > SleepSlopeLimit))
+		else if (!GVehicleDebugParams.DisableVehicleSleep && !VehicleState.bSleeping && !bControlInputPressed && VehicleState.bAllWheelsOnGround && (VehicleState.VehicleUpAxis.Z > SleepSlopeLimit))
 		{
 			float SpeedSqr = TargetInstance->GetUnrealWorldVelocity().SizeSquared();
 			if (SpeedSqr < (SleepThreshold* SleepThreshold))
