Don't wake vehicles when only the steering input value is set

#okforgithum public, internal

diff --git a/Engine/Plugins/Experimental/ChaosVehiclesPlugin/Source/ChaosVehicles/Private/ChaosVehicleMovementComponent.cpp b/Engine/Plugins/Experimental/ChaosVehiclesPlugin/Source/ChaosVehicles/Private/ChaosVehicleMovementComponent.cpp
--- a/Engine/Plugins/Experimental/ChaosVehiclesPlugin/Source/ChaosVehicles/Private/ChaosVehicleMovementComponent.cpp
+++ b/Engine/Plugins/Experimental/ChaosVehiclesPlugin/Source/ChaosVehicles/Private/ChaosVehicleMovementComponent.cpp
@@ -1240,9 +1240,9 @@ void UChaosVehicleMovementComponent::ProcessSleeping(const FControlInputs& Contr
 
 		const AController* Controller = GetController();
 		const bool bIsLocallyControlled = (Controller && Controller->IsLocalController());
-		const bool bControlInputPressed = bIsLocallyControlled ? (ControlInputs.ThrottleInput >= SMALL_NUMBER) || (FMath::Abs(ControlInputs.SteeringInput) >= SMALL_NUMBER)
+		const bool bControlInputPressed = bIsLocallyControlled ? (ControlInputs.ThrottleInput >= SMALL_NUMBER) /*|| (FMath::Abs(ControlInputs.SteeringInput) >= SMALL_NUMBER)*/
 			|| (ControlInputs.RollInput >= SMALL_NUMBER) || (ControlInputs.PitchInput >= SMALL_NUMBER) || (ControlInputs.YawInput >= SMALL_NUMBER)
-			: (ReplicatedState.ThrottleInput >= SMALL_NUMBER) || (FMath::Abs(ReplicatedState.SteeringInput) >= SMALL_NUMBER)
+			: (ReplicatedState.ThrottleInput >= SMALL_NUMBER) /*|| (FMath::Abs(ReplicatedState.SteeringInput) >= SMALL_NUMBER)*/
 			|| (ReplicatedState.RollInput >= SMALL_NUMBER) || (ReplicatedState.PitchInput >= SMALL_NUMBER) || (ReplicatedState.YawInput >= SMALL_NUMBER);
 
 		// Wake if control input pressed
