Vehicle stat counters for number of vehicles/awake/sleeping
Wake only when steering input 'changes' rather than on the value being greater than 0.

diff --git a/Engine/Plugins/Experimental/ChaosVehiclesPlugin/Source/ChaosVehicles/Private/ChaosVehicleManager.cpp b/Engine/Plugins/Experimental/ChaosVehiclesPlugin/Source/ChaosVehicles/Private/ChaosVehicleManager.cpp
--- a/Engine/Plugins/Experimental/ChaosVehiclesPlugin/Source/ChaosVehicles/Private/ChaosVehicleManager.cpp
+++ b/Engine/Plugins/Experimental/ChaosVehiclesPlugin/Source/ChaosVehicles/Private/ChaosVehicleManager.cpp
@@ -15,6 +15,9 @@ DECLARE_CYCLE_STAT(TEXT("VehicleManager:ParallelUpdateVehicles"), STAT_ChaosVehi
 DECLARE_CYCLE_STAT(TEXT("VehicleManager:Update"), STAT_ChaosVehicleManager_Update, STATGROUP_ChaosVehicleManager);
 DECLARE_CYCLE_STAT(TEXT("VehicleManager:ScenePreTick"), STAT_ChaosVehicleManager_ScenePreTick, STATGROUP_ChaosVehicleManager);
 
+DECLARE_DWORD_ACCUMULATOR_STAT(TEXT("NumVehiclesTotal"), STAT_NumVehicles_Dynamic, STATGROUP_ChaosVehicleManager);
+DECLARE_DWORD_ACCUMULATOR_STAT(TEXT("NumVehiclesAwake"), STAT_NumVehicles_Awake, STATGROUP_ChaosVehicleManager);
+DECLARE_DWORD_ACCUMULATOR_STAT(TEXT("NumVehiclesSleeping"), STAT_NumVehicles_Sleeping, STATGROUP_ChaosVehicleManager);
 
 extern FVehicleDebugParams GVehicleDebugParams;
 
@@ -197,7 +200,19 @@ void FChaosVehicleManager::Update(FPhysScene* PhysScene, float DeltaTime)
 
 void FChaosVehicleManager::PostUpdate(FChaosScene* PhysScene)
 {
-	// Unused
+	SET_DWORD_STAT(STAT_NumVehicles_Dynamic, Vehicles.Num());
+
+	int32 SleepingCount = 0;
+	for (int32 i = 0; i < Vehicles.Num(); ++i)
+	{
+		if (Vehicles[i]->VehicleState.bSleeping)
+		{
+			SleepingCount++;
+		}
+	}
+	SET_DWORD_STAT(STAT_NumVehicles_Awake, Vehicles.Num() - SleepingCount);
+	SET_DWORD_STAT(STAT_NumVehicles_Sleeping, SleepingCount);
+
 }
 
 void FChaosVehicleManager::ParallelUpdateVehicles(float DeltaSeconds)
diff --git a/Engine/Plugins/Experimental/ChaosVehiclesPlugin/Source/ChaosVehicles/Private/ChaosVehicleMovementComponent.cpp b/Engine/Plugins/Experimental/ChaosVehiclesPlugin/Source/ChaosVehicles/Private/ChaosVehicleMovementComponent.cpp
--- a/Engine/Plugins/Experimental/ChaosVehiclesPlugin/Source/ChaosVehicles/Private/ChaosVehicleMovementComponent.cpp
+++ b/Engine/Plugins/Experimental/ChaosVehiclesPlugin/Source/ChaosVehicles/Private/ChaosVehicleMovementComponent.cpp
@@ -66,6 +66,7 @@ FAutoConsoleVariableRef CVarChaosVehiclesSleepCounterThreshold(TEXT("p.Vehicle.S
 FAutoConsoleVariableRef CVarChaosVehiclesDisableVehicleSleep(TEXT("p.Vehicle.DisableVehicleSleep"), GVehicleDebugParams.DisableVehicleSleep, TEXT("Disable Vehicle Agressive Sleeping."));
 FAutoConsoleVariableRef CVarChaosVehiclesSetMaxMPH(TEXT("p.Vehicle.SetMaxMPH"), GVehicleDebugParams.SetMaxMPH, TEXT("Set a top speed in MPH (affects all vehicles)."));
 FAutoConsoleVariableRef CVarChaosVehiclesEnableMultithreading(TEXT("p.Vehicle.EnableMultithreading"), GVehicleDebugParams.EnableMultithreading, TEXT("Enable multi-threading of vehicle updates."));
+FAutoConsoleVariableRef CVarChaosVehiclesControlInputWakeTolerance(TEXT("p.Vehicle.ControlInputWakeTolerance"), GVehicleDebugParams.ControlInputWakeTolerance, TEXT("Set the control input wake tolerance."));
 
 
 void FVehicleState::CaptureState(const FBodyInstance* TargetInstance, float GravityZ, float DeltaTime)
@@ -628,6 +629,9 @@ UChaosVehicleMovementComponent::UChaosVehicleMovementComponent(const FObjectInit
 	AngErrorAccumulator = 0.0f;
 	TargetGear = 0;
 
+	PrevSteeringInput = 0.0f;
+	PrevReplicatedSteeringInput = 0.0f;
+
 	bRequiresControllerForInputs = true;
 	IdleBrakeInput = 0.0f;
 	StopThreshold = 10.0f; 
@@ -1240,10 +1244,13 @@ void UChaosVehicleMovementComponent::ProcessSleeping(const FControlInputs& Contr
 
 		const AController* Controller = GetController();
 		const bool bIsLocallyControlled = (Controller && Controller->IsLocalController());
-		const bool bControlInputPressed = bIsLocallyControlled ? (ControlInputs.ThrottleInput >= SMALL_NUMBER) || (FMath::Abs(ControlInputs.SteeringInput) >= SMALL_NUMBER)
-			|| (ControlInputs.RollInput >= SMALL_NUMBER) || (ControlInputs.PitchInput >= SMALL_NUMBER) || (ControlInputs.YawInput >= SMALL_NUMBER)
-			: (ReplicatedState.ThrottleInput >= SMALL_NUMBER) || (FMath::Abs(ReplicatedState.SteeringInput) >= SMALL_NUMBER)
-			|| (ReplicatedState.RollInput >= SMALL_NUMBER) || (ReplicatedState.PitchInput >= SMALL_NUMBER) || (ReplicatedState.YawInput >= SMALL_NUMBER);
+		const bool bControlInputPressed = bIsLocallyControlled ? (ControlInputs.ThrottleInput >= GVehicleDebugParams.ControlInputWakeTolerance) || (FMath::Abs(ControlInputs.SteeringInput - PrevSteeringInput) >= GVehicleDebugParams.ControlInputWakeTolerance)
+			|| (ControlInputs.RollInput >= GVehicleDebugParams.ControlInputWakeTolerance) || (ControlInputs.PitchInput >= GVehicleDebugParams.ControlInputWakeTolerance) || (ControlInputs.YawInput >= GVehicleDebugParams.ControlInputWakeTolerance)
+			: (ReplicatedState.ThrottleInput >= GVehicleDebugParams.ControlInputWakeTolerance) || (FMath::Abs(ReplicatedState.SteeringInput - PrevReplicatedSteeringInput) >= GVehicleDebugParams.ControlInputWakeTolerance)
+			|| (ReplicatedState.RollInput >= GVehicleDebugParams.ControlInputWakeTolerance) || (ReplicatedState.PitchInput >= GVehicleDebugParams.ControlInputWakeTolerance) || (ReplicatedState.YawInput >= GVehicleDebugParams.ControlInputWakeTolerance);
+
+		PrevSteeringInput = ControlInputs.SteeringInput;
+		PrevReplicatedSteeringInput = ReplicatedState.SteeringInput;
 
 		// Wake if control input pressed
 		if (VehicleState.bSleeping && bControlInputPressed)
diff --git a/Engine/Plugins/Experimental/ChaosVehiclesPlugin/Source/ChaosVehicles/Public/ChaosVehicleMovementComponent.h b/Engine/Plugins/Experimental/ChaosVehiclesPlugin/Source/ChaosVehicles/Public/ChaosVehicleMovementComponent.h
--- a/Engine/Plugins/Experimental/ChaosVehiclesPlugin/Source/ChaosVehicles/Public/ChaosVehicleMovementComponent.h
+++ b/Engine/Plugins/Experimental/ChaosVehiclesPlugin/Source/ChaosVehicles/Public/ChaosVehicleMovementComponent.h
@@ -96,6 +96,7 @@ struct FVehicleDebugParams
 	bool DisableVehicleSleep = false;
 	bool EnableMultithreading = true;
 	float SetMaxMPH = 0.0f;
+	float ControlInputWakeTolerance = 0.02f;
 };
 
 struct FBodyInstance;
@@ -790,6 +791,7 @@ UCLASS(Abstract, hidecategories=(PlanarMovement, "Components|Movement|Planar", A
 class CHAOSVEHICLES_API UChaosVehicleMovementComponent : public UPawnMovementComponent
 {
 	friend struct FChaosVehicleDefaultAsyncInput;
+	friend class FChaosVehicleManager;
 
 	GENERATED_UCLASS_BODY()
 
@@ -1343,4 +1345,7 @@ private:
 	Chaos::FSimpleAerodynamicsConfig PAerodynamicsSetup;
 	int32 TargetGear;
 
+	float PrevSteeringInput;
+	float PrevReplicatedSteeringInput;
+
 };
