Add ResetVehicle function to be called if reusing vehicle

diff --git a/Engine/Plugins/Experimental/ChaosVehiclesPlugin/Source/ChaosVehicles/Private/ChaosVehicleMovementComponent.cpp b/Engine/Plugins/Experimental/ChaosVehiclesPlugin/Source/ChaosVehicles/Private/ChaosVehicleMovementComponent.cpp
--- a/Engine/Plugins/Experimental/ChaosVehiclesPlugin/Source/ChaosVehicles/Private/ChaosVehicleMovementComponent.cpp
+++ b/Engine/Plugins/Experimental/ChaosVehiclesPlugin/Source/ChaosVehicles/Private/ChaosVehicleMovementComponent.cpp
@@ -1769,6 +1769,11 @@ void UChaosVehicleMovementComponent::Update(float DeltaTime)
 #endif
 }
 
+void UChaosVehicleMovementComponent::ResetVehicleState()
+{
+	ClearRawInput();
+}
+
 void UChaosVehicleMovementComponent::FinalizeSimCallbackData(FChaosVehicleManagerAsyncInput& Input)
 {
 	bool bIsPhysicsStateCreated = true;
diff --git a/Engine/Plugins/Experimental/ChaosVehiclesPlugin/Source/ChaosVehicles/Private/ChaosWheeledVehicleMovementComponent.cpp b/Engine/Plugins/Experimental/ChaosVehiclesPlugin/Source/ChaosVehicles/Private/ChaosWheeledVehicleMovementComponent.cpp
--- a/Engine/Plugins/Experimental/ChaosVehiclesPlugin/Source/ChaosVehicles/Private/ChaosWheeledVehicleMovementComponent.cpp
+++ b/Engine/Plugins/Experimental/ChaosVehiclesPlugin/Source/ChaosVehicles/Private/ChaosWheeledVehicleMovementComponent.cpp
@@ -1310,6 +1310,7 @@ void UChaosWheeledVehicleMovementComponent::CreateWheels()
 	}
 
 	WheelStatus.SetNum(WheelSetups.Num());
+	CachedState.SetNum(WheelSetups.Num());
 
 	RecalculateAxles();
 }
@@ -1441,6 +1442,21 @@ void UChaosWheeledVehicleMovementComponent::SetupVehicle(TUniquePtr<Chaos::FSimp
 	SetupSuspension(PVehicle);
 }
 
+void UChaosWheeledVehicleMovementComponent::ResetVehicleState()
+{
+	UChaosVehicleMovementComponent::ResetVehicleState();
+
+	for (FWheelStatus& WheelInfo : WheelStatus)
+	{
+		WheelInfo.Init();
+	}
+	
+	for (FCachedState& State : CachedState)
+	{
+		State.bIsValid = false;
+	}
+}
+
 void UChaosWheeledVehicleMovementComponent::SetupVehicleShapes()
 {
 	if (!UpdatedPrimitive)
@@ -2559,47 +2575,60 @@ float UChaosWheeledVehicleMovementComponent::GetSuspensionOffset(int WheelIndex)
 				{
 					Offset = -Wheel->SuspensionMaxDrop;
 				}
+
+				CachedState[WheelIndex].bIsValid = true;
+				CachedState[WheelIndex].WheelOffset = Offset;
+
 			}
 			else
 			{
-				ECollisionChannel SpringCollisionChannel = ECollisionChannel::ECC_WorldDynamic;
-				FCollisionResponseParams ResponseParams;
-				ResponseParams.CollisionResponse = WheelTraceCollisionResponses;
+				if (VehicleState.bSleeping && CachedState[WheelIndex].bIsValid)
+				{
+					Offset = CachedState[WheelIndex].WheelOffset;
+				}
+				else
+				{
+					ECollisionChannel SpringCollisionChannel = ECollisionChannel::ECC_WorldDynamic;
+					FCollisionResponseParams ResponseParams;
+					ResponseParams.CollisionResponse = WheelTraceCollisionResponses;
 
-				TArray<AActor*> ActorsToIgnore;
-				ActorsToIgnore.Add(GetPawnOwner()); // ignore self in scene query
+					TArray<AActor*> ActorsToIgnore;
+					ActorsToIgnore.Add(GetPawnOwner()); // ignore self in scene query
 
-				FCollisionQueryParams TraceParams(NAME_None, FCollisionQueryParams::GetUnknownStatId(), false, nullptr);
-				TraceParams.bReturnPhysicalMaterial = true;	// we need this to get the surface friction coefficient
-				TraceParams.AddIgnoredActors(ActorsToIgnore);
-				TraceParams.bTraceComplex = (Wheels[WheelIndex]->SweepType == ESweepType::ComplexSweep);
+					FCollisionQueryParams TraceParams(NAME_None, FCollisionQueryParams::GetUnknownStatId(), false, nullptr);
+					TraceParams.bReturnPhysicalMaterial = true;	// we need this to get the surface friction coefficient
+					TraceParams.AddIgnoredActors(ActorsToIgnore);
+					TraceParams.bTraceComplex = (Wheels[WheelIndex]->SweepType == ESweepType::ComplexSweep);
 
-				FVector LocalDirection = Wheel->SuspensionAxis;
-				FVector WorldLocation = VehicleWorldTransform.TransformPosition(GetWheelRestingPosition(WheelSetup));
-				FVector WorldDirection = VehicleWorldTransform.TransformVector(LocalDirection);
+					FVector LocalDirection = Wheel->SuspensionAxis;
+					FVector WorldLocation = VehicleWorldTransform.TransformPosition(GetWheelRestingPosition(WheelSetup));
+					FVector WorldDirection = VehicleWorldTransform.TransformVector(LocalDirection);
 
-				FVector TraceStart = WorldLocation - WorldDirection * (Wheel->SuspensionMaxRaise);
-				FVector TraceEnd = WorldLocation + WorldDirection * (Wheel->SuspensionMaxDrop + Wheel->WheelRadius);
+					FVector TraceStart = WorldLocation - WorldDirection * (Wheel->SuspensionMaxRaise);
+					FVector TraceEnd = WorldLocation + WorldDirection * (Wheel->SuspensionMaxDrop + Wheel->WheelRadius);
 
-				FHitResult HitResult;
-				GetWorld()->LineTraceSingleByChannel(HitResult, TraceStart, TraceEnd, SpringCollisionChannel, TraceParams, ResponseParams);
+					FHitResult HitResult;
+					GetWorld()->LineTraceSingleByChannel(HitResult, TraceStart, TraceEnd, SpringCollisionChannel, TraceParams, ResponseParams);
 
-				if (HitResult.bBlockingHit)
-				{
-					FVector LocalPos = GetWheelRestingPosition(WheelSetup);
-					FVector LocalHitPoint = VehicleWorldTransform.InverseTransformPosition(HitResult.ImpactPoint);
-					Offset = LocalHitPoint.Z - LocalPos.Z + Wheel->WheelRadius;
-					Offset = FMath::Clamp(Offset, -Wheel->SuspensionMaxDrop, Wheel->SuspensionMaxRaise);
-				}
-				else
-				{
-					Offset = -Wheel->SuspensionMaxDrop;
-				}
+					if (HitResult.bBlockingHit)
+					{
+						FVector LocalPos = GetWheelRestingPosition(WheelSetup);
+						FVector LocalHitPoint = VehicleWorldTransform.InverseTransformPosition(HitResult.ImpactPoint);
+						Offset = LocalHitPoint.Z - LocalPos.Z + Wheel->WheelRadius;
+						Offset = FMath::Clamp(Offset, -Wheel->SuspensionMaxDrop, Wheel->SuspensionMaxRaise);
+					}
+					else
+					{
+						Offset = -Wheel->SuspensionMaxDrop;
+					}
 
+					CachedState[WheelIndex].bIsValid = true;
+					CachedState[WheelIndex].WheelOffset = Offset;
+				}
 			}
 		}
 	}
-
+	
 	return Offset;
 }
 
diff --git a/Engine/Plugins/Experimental/ChaosVehiclesPlugin/Source/ChaosVehicles/Public/ChaosVehicleMovementComponent.h b/Engine/Plugins/Experimental/ChaosVehiclesPlugin/Source/ChaosVehicles/Public/ChaosVehicleMovementComponent.h
--- a/Engine/Plugins/Experimental/ChaosVehiclesPlugin/Source/ChaosVehicles/Public/ChaosVehicleMovementComponent.h
+++ b/Engine/Plugins/Experimental/ChaosVehiclesPlugin/Source/ChaosVehicles/Public/ChaosVehicleMovementComponent.h
@@ -1041,6 +1041,10 @@ public:
 	UFUNCTION(BlueprintCallable, Category = "Game|Components|ChaosVehicleMovement")
 	bool IsParked() const;
 
+	/** Reset some vehicle state - call this if you are say creating pool of vehicles that get reused and you don't want to carry over the previous state */
+	UFUNCTION(BlueprintCallable, Category = "Game|Components|ChaosVehicleMovement")
+	void ResetVehicle() { ResetVehicleState(); }
+
 	/** Grab a snapshot of the vehicle instance dynamic state */
 	void GetBaseSnapshot(FBaseSnapshotData& SnapshotOut) const;
 
@@ -1073,6 +1077,8 @@ public:
 
 	virtual void Update(float DeltaTime);
 
+	virtual void ResetVehicleState();
+
 	// Get output data from Physics Thread
 	virtual void ParallelUpdate(float DeltaSeconds);
 
diff --git a/Engine/Plugins/Experimental/ChaosVehiclesPlugin/Source/ChaosVehicles/Public/ChaosWheeledVehicleMovementComponent.h b/Engine/Plugins/Experimental/ChaosVehiclesPlugin/Source/ChaosVehicles/Public/ChaosWheeledVehicleMovementComponent.h
--- a/Engine/Plugins/Experimental/ChaosVehiclesPlugin/Source/ChaosVehicles/Public/ChaosWheeledVehicleMovementComponent.h
+++ b/Engine/Plugins/Experimental/ChaosVehiclesPlugin/Source/ChaosVehicles/Public/ChaosWheeledVehicleMovementComponent.h
@@ -840,6 +840,8 @@ class CHAOSVEHICLES_API UChaosWheeledVehicleMovementComponent : public UChaosVeh
 	/** Allocate and setup the Chaos vehicle */
 	virtual void SetupVehicle(TUniquePtr<Chaos::FSimpleWheeledVehicle>& PVehicle) override;
 
+	virtual void ResetVehicleState() override;
+
 protected:
 
 	//////////////////////////////////////////////////////////////////////////
@@ -902,13 +904,22 @@ protected:
 	void DrawDial(UCanvas* Canvas, FVector2D Pos, float Radius, float CurrentValue, float MaxValue);
 #endif
 
+	struct FCachedState
+	{
+		FCachedState() : WheelOffset(0.f), bIsValid(false)
+		{ }
+
+		float WheelOffset;
+		bool bIsValid;
+	};
+
 	static EDebugPages DebugPage;
 	uint32 NumDrivenWheels; /** The number of wheels that have engine enabled checked */
 	FVector2D WheelTrackDimensions;	// Wheelbase (X) and track (Y) dimensions
 	TMap<UChaosVehicleWheel*, TArray<int>> AxleToWheelMap;
 	TArray<FPhysicsConstraintHandle> ConstraintHandles;
 	TArray<FWheelStatus> WheelStatus; /** Wheel output status */
-
+	TArray<FCachedState> CachedState;
 	Chaos::FPerformanceMeasure PerformanceMeasure;
 };
 
