Add more clearing of state in ResetVehicle function

diff --git a/Engine/Plugins/Experimental/ChaosVehiclesPlugin/Source/ChaosVehicles/Private/ChaosVehicleMovementComponent.cpp b/Engine/Plugins/Experimental/ChaosVehiclesPlugin/Source/ChaosVehicles/Private/ChaosVehicleMovementComponent.cpp
--- a/Engine/Plugins/Experimental/ChaosVehiclesPlugin/Source/ChaosVehicles/Private/ChaosVehicleMovementComponent.cpp
+++ b/Engine/Plugins/Experimental/ChaosVehiclesPlugin/Source/ChaosVehicles/Private/ChaosVehicleMovementComponent.cpp
@@ -1772,6 +1772,8 @@ void UChaosVehicleMovementComponent::Update(float DeltaTime)
 void UChaosVehicleMovementComponent::ResetVehicleState()
 {
 	ClearRawInput();
+	StopMovementImmediately();
+	SetTargetGear(0, true);
 }
 
 void UChaosVehicleMovementComponent::FinalizeSimCallbackData(FChaosVehicleManagerAsyncInput& Input)
