Bug Fix for parked vehicles not driving properly when recycled by the pooling system

#FROST-4235

diff --git a/Engine/Plugins/Experimental/ChaosVehiclesPlugin/Source/ChaosVehicles/Private/ChaosVehicleMovementComponent.cpp b/Engine/Plugins/Experimental/ChaosVehiclesPlugin/Source/ChaosVehicles/Private/ChaosVehicleMovementComponent.cpp
--- a/Engine/Plugins/Experimental/ChaosVehiclesPlugin/Source/ChaosVehicles/Private/ChaosVehicleMovementComponent.cpp
+++ b/Engine/Plugins/Experimental/ChaosVehiclesPlugin/Source/ChaosVehicles/Private/ChaosVehicleMovementComponent.cpp
@@ -1773,7 +1773,11 @@ void UChaosVehicleMovementComponent::ResetVehicleState()
 {
 	ClearRawInput();
 	StopMovementImmediately();
-	SetTargetGear(0, true);
+
+	OnDestroyPhysicsState();
+	OnCreatePhysicsState();
+
+	SetTargetGear(1, true);
 }
 
 void UChaosVehicleMovementComponent::FinalizeSimCallbackData(FChaosVehicleManagerAsyncInput& Input)
diff --git a/Engine/Plugins/Experimental/ChaosVehiclesPlugin/Source/ChaosVehicles/Private/ChaosWheeledVehicleMovementComponent.cpp b/Engine/Plugins/Experimental/ChaosVehiclesPlugin/Source/ChaosVehicles/Private/ChaosWheeledVehicleMovementComponent.cpp
--- a/Engine/Plugins/Experimental/ChaosVehiclesPlugin/Source/ChaosVehicles/Private/ChaosWheeledVehicleMovementComponent.cpp
+++ b/Engine/Plugins/Experimental/ChaosVehiclesPlugin/Source/ChaosVehicles/Private/ChaosWheeledVehicleMovementComponent.cpp
@@ -1330,6 +1330,7 @@ void UChaosWheeledVehicleMovementComponent::SetupVehicle(TUniquePtr<Chaos::FSimp
 	check(PVehicle);
 
 	Super::SetupVehicle(PVehicle);
+	NumDrivenWheels = 0;
 
 	// we are allowed any number of wheels not limited to only 4
 	for (int32 WheelIdx = 0; WheelIdx < WheelSetups.Num(); ++WheelIdx)
