Fix pooled vehicles not resetting transmission properly

diff --git a/Engine/Plugins/Experimental/ChaosVehiclesPlugin/Source/ChaosVehicles/Private/ChaosVehicleMovementComponent.cpp b/Engine/Plugins/Experimental/ChaosVehiclesPlugin/Source/ChaosVehicles/Private/ChaosVehicleMovementComponent.cpp
--- a/Engine/Plugins/Experimental/ChaosVehiclesPlugin/Source/ChaosVehicles/Private/ChaosVehicleMovementComponent.cpp
+++ b/Engine/Plugins/Experimental/ChaosVehiclesPlugin/Source/ChaosVehicles/Private/ChaosVehicleMovementComponent.cpp
@@ -1777,7 +1777,9 @@ void UChaosVehicleMovementComponent::ResetVehicleState()
 	OnDestroyPhysicsState();
 	OnCreatePhysicsState();
 
-	SetTargetGear(1, true);
+	// Shift into neutral, force the local copy of target gear to be correct
+	SetTargetGear(0, true);
+	TargetGear = 0;
 }
 
 void UChaosVehicleMovementComponent::FinalizeSimCallbackData(FChaosVehicleManagerAsyncInput& Input)
