Add a setter for RequiresControllerForInputs variable

diff --git a/Engine/Plugins/Experimental/ChaosVehiclesPlugin/Source/ChaosVehicles/Private/ChaosVehicleMovementComponent.cpp b/Engine/Plugins/Experimental/ChaosVehiclesPlugin/Source/ChaosVehicles/Private/ChaosVehicleMovementComponent.cpp
--- a/Engine/Plugins/Experimental/ChaosVehiclesPlugin/Source/ChaosVehicles/Private/ChaosVehicleMovementComponent.cpp
+++ b/Engine/Plugins/Experimental/ChaosVehiclesPlugin/Source/ChaosVehicles/Private/ChaosVehicleMovementComponent.cpp
@@ -955,6 +955,11 @@ void UChaosVehicleMovementComponent::SetUseAutomaticGears(bool bUseAuto)
 	TransmissionType = bUseAuto ? Chaos::ETransmissionType::Automatic : Chaos::ETransmissionType::Manual;
 }
 
+void UChaosVehicleMovementComponent::SetRequiresControllerForInputs(bool bRequiresController)
+{
+	bRequiresControllerForInputs = bRequiresController;
+}
+
 // Data access
 
 int32 UChaosVehicleMovementComponent::GetCurrentGear() const
diff --git a/Engine/Plugins/Experimental/ChaosVehiclesPlugin/Source/ChaosVehicles/Public/ChaosVehicleMovementComponent.h b/Engine/Plugins/Experimental/ChaosVehiclesPlugin/Source/ChaosVehicles/Public/ChaosVehicleMovementComponent.h
--- a/Engine/Plugins/Experimental/ChaosVehiclesPlugin/Source/ChaosVehicles/Public/ChaosVehicleMovementComponent.h
+++ b/Engine/Plugins/Experimental/ChaosVehiclesPlugin/Source/ChaosVehicles/Public/ChaosVehicleMovementComponent.h
@@ -1000,6 +1000,10 @@ public:
 	UFUNCTION(BlueprintCallable, Category="Game|Components|ChaosVehicleMovement")
 	void SetUseAutomaticGears(bool bUseAuto);
 
+	/** Set the flag that determines whether a controller is required to set control inputs */
+	UFUNCTION(BlueprintCallable, Category="Game|Components|ChaosVehicleMovement")
+	void SetRequiresControllerForInputs(bool bRequiresController);
+
 	/** Get current gear */
 	UFUNCTION(BlueprintCallable, Category = "Game|Components|ChaosVehicleMovement")
 	int32 GetCurrentGear() const;
