New StagecoachWheelController for handling wheel rotation for motion blur effect.

diff --git a/Engine/Plugins/Experimental/ChaosVehiclesPlugin/Source/ChaosVehicles/Private/AnimNode_StageCoachWheelController.cpp b/Engine/Plugins/Experimental/ChaosVehiclesPlugin/Source/ChaosVehicles/Private/AnimNode_StageCoachWheelController.cpp
new file mode 100644
--- /dev/null
+++ b/Engine/Plugins/Experimental/ChaosVehiclesPlugin/Source/ChaosVehicles/Private/AnimNode_StageCoachWheelController.cpp
@@ -0,0 +1,143 @@
+// Copyright Epic Games, Inc. All Rights Reserved.
+
+#include "AnimNode_StageCoachWheelController.h"
+#include "AnimationRuntime.h"
+#include "WheeledVehiclePawn.h"
+#include "ChaosVehicleMovementComponent.h"
+
+
+FAnimNode_StageCoachWheelController::FAnimNode_StageCoachWheelController()
+	: WheelSpokeCount(5)
+	, MaxAngularVelocity(256.f)
+	, ShutterSpeed(30.f)
+	, StageCoachBlend(730.f)
+{
+	AnimInstanceProxy = nullptr;
+}
+
+void FAnimNode_StageCoachWheelController::GatherDebugData(FNodeDebugData& DebugData)
+{
+	FString DebugLine = DebugData.GetNodeName(this);
+
+	DebugLine += "(";
+	AddDebugNodeData(DebugLine);
+	DebugLine += ")";
+
+	DebugData.AddDebugItem(DebugLine);
+
+	const TArray<FWheelAnimationData>& WheelAnimData = AnimInstanceProxy->GetWheelAnimData();
+	for (const FWheelLookupData& Wheel : Wheels)
+	{
+		if (Wheel.BoneReference.BoneIndex != INDEX_NONE)
+		{
+			DebugLine = FString::Printf(TEXT(" [Wheel Index : %d] Bone: %s , Rotation Offset : %s, Location Offset : %s"),
+				Wheel.WheelIndex, *Wheel.BoneReference.BoneName.ToString(), *WheelAnimData[Wheel.WheelIndex].RotOffset.ToString(), *WheelAnimData[Wheel.WheelIndex].LocOffset.ToString());
+		}
+		else
+		{
+			DebugLine = FString::Printf(TEXT(" [Wheel Index : %d] Bone: %s (invalid bone)"),
+				Wheel.WheelIndex, *Wheel.BoneReference.BoneName.ToString());
+		}
+
+		DebugData.AddDebugItem(DebugLine);
+	}
+
+	ComponentPose.GatherDebugData(DebugData);
+}
+
+void FAnimNode_StageCoachWheelController::EvaluateSkeletalControl_AnyThread(FComponentSpacePoseContext& Output, TArray<FBoneTransform>& OutBoneTransforms)
+{
+	check(OutBoneTransforms.Num() == 0);
+
+	const TArray<FWheelAnimationData>& WheelAnimData = AnimInstanceProxy->GetWheelAnimData();
+
+	const FBoneContainer& BoneContainer = Output.Pose.GetPose().GetBoneContainer();
+	for(const FWheelLookupData& Wheel : Wheels)
+	{
+		if (Wheel.BoneReference.IsValidToEvaluate(BoneContainer))
+		{
+			FCompactPoseBoneIndex WheelSimBoneIndex = Wheel.BoneReference.GetCompactPoseIndex(BoneContainer);
+
+			// the way we apply transform is same as FMatrix or FTransform
+			// we apply scale first, and rotation, and translation
+			// if you'd like to translate first, you'll need two nodes that first node does translate and second nodes to rotate. 
+			FTransform NewBoneTM = Output.Pose.GetComponentSpaceTransform(WheelSimBoneIndex);
+
+			FAnimationRuntime::ConvertCSTransformToBoneSpace(Output.AnimInstanceProxy->GetComponentTransform(), Output.Pose, NewBoneTM, WheelSimBoneIndex, BCS_ComponentSpace);
+
+			// Apply rotation offset
+			const FQuat BoneQuat(WheelAnimData[Wheel.WheelIndex].RotOffset);
+			NewBoneTM.SetRotation(BoneQuat * NewBoneTM.GetRotation());
+
+			// Apply loc offset
+			NewBoneTM.AddToTranslation(WheelAnimData[Wheel.WheelIndex].LocOffset);
+
+			// Convert back to Component Space.
+			FAnimationRuntime::ConvertBoneSpaceTransformToCS(Output.AnimInstanceProxy->GetComponentTransform(), Output.Pose, NewBoneTM, WheelSimBoneIndex, BCS_ComponentSpace);
+
+			// add back to it
+			OutBoneTransforms.Add(FBoneTransform(WheelSimBoneIndex, NewBoneTM));
+		}
+	}
+
+#if ANIM_TRACE_ENABLED
+	for (const FWheelLookupData& Wheel : Wheels)
+	{
+		if (Wheel.BoneReference.BoneIndex != INDEX_NONE)
+		{
+			TRACE_ANIM_NODE_VALUE(Output, *FString::Printf(TEXT("Wheel %d Name"), Wheel.WheelIndex), *Wheel.BoneReference.BoneName.ToString());
+			TRACE_ANIM_NODE_VALUE(Output, *FString::Printf(TEXT("Wheel %d Rotation Offset"), Wheel.WheelIndex), WheelAnimData[Wheel.WheelIndex].RotOffset);
+			TRACE_ANIM_NODE_VALUE(Output, *FString::Printf(TEXT("Wheel %d Location Offset"), Wheel.WheelIndex), WheelAnimData[Wheel.WheelIndex].LocOffset);
+		}
+		else
+		{
+			TRACE_ANIM_NODE_VALUE(Output, *FString::Printf(TEXT("Wheel %d Name"), Wheel.WheelIndex), *FString::Printf(TEXT("%s (invalid)"), *Wheel.BoneReference.BoneName.ToString()));
+		}
+	}
+#endif
+}
+
+bool FAnimNode_StageCoachWheelController::IsValidToEvaluate(const USkeleton* Skeleton, const FBoneContainer& RequiredBones) 
+{
+	// if both bones are valid
+	for (const FWheelLookupData& Wheel : Wheels)
+	{
+		// if one of them is valid
+		if (Wheel.BoneReference.IsValidToEvaluate(RequiredBones) == true)
+		{
+			return true;
+		}
+	}
+
+	return false;
+}
+
+void FAnimNode_StageCoachWheelController::InitializeBoneReferences(const FBoneContainer& RequiredBones) 
+{
+	const TArray<FWheelAnimationData>& WheelAnimData = AnimInstanceProxy->GetWheelAnimData();
+	const int32 NumWheels = WheelAnimData.Num();
+	Wheels.Empty(NumWheels);
+
+	for (int32 WheelIndex = 0; WheelIndex < NumWheels; ++WheelIndex)
+	{
+		FWheelLookupData* Wheel = new(Wheels)FWheelLookupData();
+		Wheel->WheelIndex = WheelIndex;
+		Wheel->BoneReference.BoneName = WheelAnimData[WheelIndex].BoneName;
+		Wheel->BoneReference.Initialize(RequiredBones);
+	}
+
+	// sort by bone indices
+	Wheels.Sort([](const FWheelLookupData& L, const FWheelLookupData& R) { return L.BoneReference.BoneIndex < R.BoneReference.BoneIndex; });
+}
+
+void FAnimNode_StageCoachWheelController::Initialize_AnyThread(const FAnimationInitializeContext& Context)
+{
+	AnimInstanceProxy = (FVehicleAnimationInstanceProxy*)Context.AnimInstanceProxy;
+
+	if (AnimInstanceProxy)
+	{
+		AnimInstanceProxy->SetStageCoachEffectParams(WheelSpokeCount, MaxAngularVelocity, ShutterSpeed, StageCoachBlend);
+	}
+}
+
+
diff --git a/Engine/Plugins/Experimental/ChaosVehiclesPlugin/Source/ChaosVehicles/Private/ChaosVehicleWheel.cpp b/Engine/Plugins/Experimental/ChaosVehiclesPlugin/Source/ChaosVehicles/Private/ChaosVehicleWheel.cpp
--- a/Engine/Plugins/Experimental/ChaosVehiclesPlugin/Source/ChaosVehicles/Private/ChaosVehicleWheel.cpp
+++ b/Engine/Plugins/Experimental/ChaosVehiclesPlugin/Source/ChaosVehicles/Private/ChaosVehicleWheel.cpp
@@ -72,6 +72,15 @@ float UChaosVehicleWheel::GetRotationAngle() const
 	return RotationAngle;
 }
 
+float UChaosVehicleWheel::GetRotationAngularVelocity() const
+{
+	check(VehicleComponent && VehicleComponent->PhysicsVehicleOutput());
+	float RotationAngularVelocity = -1.0f * FMath::RadiansToDegrees(VehicleComponent->PhysicsVehicleOutput()->Wheels[WheelIndex].AngularVelocity);
+	ensure(!FMath::IsNaN(RotationAngularVelocity));
+	return RotationAngularVelocity;
+}
+
+
 float UChaosVehicleWheel::GetWheelRadius() const
 {
 	check(VehicleComponent && VehicleComponent->PhysicsVehicleOutput());
diff --git a/Engine/Plugins/Experimental/ChaosVehiclesPlugin/Source/ChaosVehicles/Private/VehicleAnimationInstance.cpp b/Engine/Plugins/Experimental/ChaosVehiclesPlugin/Source/ChaosVehicles/Private/VehicleAnimationInstance.cpp
--- a/Engine/Plugins/Experimental/ChaosVehiclesPlugin/Source/ChaosVehicles/Private/VehicleAnimationInstance.cpp
+++ b/Engine/Plugins/Experimental/ChaosVehiclesPlugin/Source/ChaosVehicles/Private/VehicleAnimationInstance.cpp
@@ -88,7 +88,43 @@
 				{
 					if (const UChaosVehicleWheel* VehicleWheel = WheeledVehicleComponent->Wheels[WheelIndex])
 					{
-						WheelInstance.RotOffset.Pitch = VehicleWheel->GetRotationAngle();
+						if (WheelSpokeCount > 0 && ShutterSpeed > 0 && MaxAngularVelocity > SMALL_NUMBER) // employ stagecoach effect
+						{
+							// normalized spoke transition value
+							float AngularVelocity = VehicleWheel->GetRotationAngularVelocity();
+							float DegreesPerFrame = AngularVelocity / ShutterSpeed;
+							float DegreesPerSpoke = 360.f / WheelSpokeCount;
+
+							float IntegerPart = 0;
+							float SpokeTransition = FMath::Modf(DegreesPerFrame / DegreesPerSpoke, &IntegerPart);
+							float StagecoachEffectVelocity = FMath::Sin(SpokeTransition * TWO_PI) * MaxAngularVelocity;
+
+							// blend
+							float OffsetVelocity = FMath::Abs(AngularVelocity) - MaxAngularVelocity;
+							if (OffsetVelocity < 0.f)
+							{
+								OffsetVelocity = 0.f;
+							}
+
+							float BlendAlpha = FMath::Clamp(OffsetVelocity / MaxAngularVelocity, 0.f, 1.f);
+
+							float CorrectedAngularVelocity = FMath::Lerp(AngularVelocity, StagecoachEffectVelocity, BlendAlpha);
+							CorrectedAngularVelocity = FMath::Clamp(CorrectedAngularVelocity, -MaxAngularVelocity, MaxAngularVelocity);
+
+							// integrate to angular position
+							float RotationDelta = CorrectedAngularVelocity * DeltaSeconds;
+							WheelInstance.RotOffset.Pitch += RotationDelta;
+
+							int ExcessRotations = (int)(WheelInstance.RotOffset.Pitch / 360.0f);
+							if (FMath::Abs(ExcessRotations) > 1)
+							{
+								WheelInstance.RotOffset.Pitch -= ExcessRotations * 360.0f;
+							}
+						}
+						else
+						{
+							WheelInstance.RotOffset.Pitch = VehicleWheel->GetRotationAngle();
+						}
 						WheelInstance.RotOffset.Yaw = VehicleWheel->GetSteerAngle();
 						WheelInstance.RotOffset.Roll = 0.f;
 
diff --git a/Engine/Plugins/Experimental/ChaosVehiclesPlugin/Source/ChaosVehicles/Public/AnimNode_StageCoachWheelController.h b/Engine/Plugins/Experimental/ChaosVehiclesPlugin/Source/ChaosVehicles/Public/AnimNode_StageCoachWheelController.h
new file mode 100644
--- /dev/null
+++ b/Engine/Plugins/Experimental/ChaosVehiclesPlugin/Source/ChaosVehicles/Public/AnimNode_StageCoachWheelController.h
@@ -0,0 +1,63 @@
+// Copyright Epic Games, Inc. All Rights Reserved.
+
+#pragma once
+
+#include "CoreMinimal.h"
+#include "UObject/ObjectMacros.h"
+#include "BoneContainer.h"
+#include "BonePose.h"
+#include "VehicleAnimationInstance.h"
+#include "BoneControllers/AnimNode_SkeletalControlBase.h"
+#include "AnimNode_StageCoachWheelController.generated.h"
+
+/**
+ *	Simple controller that replaces or adds to the translation/rotation of a single bone.
+ */
+USTRUCT()
+struct CHAOSVEHICLES_API FAnimNode_StageCoachWheelController : public FAnimNode_SkeletalControlBase
+{
+	GENERATED_USTRUCT_BODY()
+
+	// Number of spokes visible on wheel
+	UPROPERTY(EditAnywhere, BlueprintReadWrite, Category = Alpha, meta = (PinShownByDefault))
+	int WheelSpokeCount;
+
+	// Wheel max rotation speed in degrees/second
+	UPROPERTY(EditAnywhere, BlueprintReadWrite, Category = Alpha, meta = (PinShownByDefault))
+	float MaxAngularVelocity;
+
+	// Camera shutter speed in frames/second
+	UPROPERTY(EditAnywhere, BlueprintReadWrite, Category = Alpha, meta = (PinShownByDefault))
+	float ShutterSpeed;
+
+	// Blend effect degrees/second
+	UPROPERTY(EditAnywhere, BlueprintReadWrite, Category = Alpha, meta = (PinShownByDefault))
+	float StageCoachBlend;
+
+
+	FAnimNode_StageCoachWheelController();
+
+	// FAnimNode_Base interface
+	virtual void GatherDebugData(FNodeDebugData& DebugData) override;
+	// End of FAnimNode_Base interface
+
+	// FAnimNode_SkeletalControlBase interface
+	virtual void EvaluateSkeletalControl_AnyThread(FComponentSpacePoseContext& Output, TArray<FBoneTransform>& OutBoneTransforms) override;
+	virtual bool IsValidToEvaluate(const USkeleton* Skeleton, const FBoneContainer& RequiredBones) override;
+	virtual void Initialize_AnyThread(const FAnimationInitializeContext& Context) override;
+	// End of FAnimNode_SkeletalControlBase interface
+
+private:
+	// FAnimNode_SkeletalControlBase interface
+	virtual void InitializeBoneReferences(const FBoneContainer& RequiredBones) override;
+	// End of FAnimNode_SkeletalControlBase interface
+
+	struct FWheelLookupData
+	{
+		int32 WheelIndex;
+		FBoneReference BoneReference;
+	};
+
+	TArray<FWheelLookupData> Wheels;
+	FVehicleAnimationInstanceProxy* AnimInstanceProxy;
+};
diff --git a/Engine/Plugins/Experimental/ChaosVehiclesPlugin/Source/ChaosVehicles/Public/ChaosVehicleWheel.h b/Engine/Plugins/Experimental/ChaosVehiclesPlugin/Source/ChaosVehicles/Public/ChaosVehicleWheel.h
--- a/Engine/Plugins/Experimental/ChaosVehiclesPlugin/Source/ChaosVehicles/Public/ChaosVehicleWheel.h
+++ b/Engine/Plugins/Experimental/ChaosVehiclesPlugin/Source/ChaosVehicles/Public/ChaosVehicleWheel.h
@@ -248,6 +248,9 @@ class UChaosWheeledVehicleMovementComponent;
 		UFUNCTION(BlueprintCallable, Category = "Game|Components|WheeledVehicleMovement")
 		float GetRotationAngle() const;
 
+		UFUNCTION(BlueprintCallable, Category = "Game|Components|WheeledVehicleMovement")
+		float GetRotationAngularVelocity() const;
+
 		UFUNCTION(BlueprintCallable, Category = "Game|Components|WheeledVehicleMovement")
 		float GetSuspensionOffset() const;
 
diff --git a/Engine/Plugins/Experimental/ChaosVehiclesPlugin/Source/ChaosVehicles/Public/VehicleAnimationInstance.h b/Engine/Plugins/Experimental/ChaosVehiclesPlugin/Source/ChaosVehicles/Public/VehicleAnimationInstance.h
--- a/Engine/Plugins/Experimental/ChaosVehiclesPlugin/Source/ChaosVehicles/Public/VehicleAnimationInstance.h
+++ b/Engine/Plugins/Experimental/ChaosVehiclesPlugin/Source/ChaosVehicles/Public/VehicleAnimationInstance.h
@@ -25,11 +25,19 @@ USTRUCT()
 
 		FVehicleAnimationInstanceProxy()
 		: FAnimInstanceProxy()
+		, WheelSpokeCount(0)
+		, MaxAngularVelocity(256.f)
+		, ShutterSpeed(30.f)
+		, StageCoachBlend(730.f)
 	{
 	}
 
 	FVehicleAnimationInstanceProxy(UAnimInstance* Instance)
 		: FAnimInstanceProxy(Instance)
+		, WheelSpokeCount(0)
+		, MaxAngularVelocity(256.f)
+		, ShutterSpeed(30.f)
+		, StageCoachBlend(730.f)
 	{
 	}
 
@@ -46,8 +54,22 @@ public:
 		return WheelInstances;
 	}
 
+	void SetStageCoachEffectParams(int InWheelSpokeCount, float InMaxAngularVelocity, float InShutterSpeed, float InStageCoachBlend)
+	{
+		WheelSpokeCount = InWheelSpokeCount;
+		MaxAngularVelocity = InMaxAngularVelocity;
+		ShutterSpeed = InShutterSpeed;
+		StageCoachBlend = InStageCoachBlend;
+	}
+
 private:
 	TArray<FWheelAnimationData> WheelInstances;
+
+	int WheelSpokeCount;			// Number of spokes visible on wheel
+	float MaxAngularVelocity;		// Wheel max rotation speed in degrees/second
+	float ShutterSpeed;				// Camera shutter speed in frames/second
+	float StageCoachBlend;			// Blend effect degrees/second
+
 };
 
 UCLASS(transient)
diff --git a/Engine/Plugins/Experimental/ChaosVehiclesPlugin/Source/ChaosVehiclesEditor/Private/AnimGraphNode_StageCoachWheelController.cpp b/Engine/Plugins/Experimental/ChaosVehiclesPlugin/Source/ChaosVehiclesEditor/Private/AnimGraphNode_StageCoachWheelController.cpp
new file mode 100644
--- /dev/null
+++ b/Engine/Plugins/Experimental/ChaosVehiclesPlugin/Source/ChaosVehiclesEditor/Private/AnimGraphNode_StageCoachWheelController.cpp
@@ -0,0 +1,59 @@
+// Copyright Epic Games, Inc. All Rights Reserved.
+
+#include "AnimGraphNode_StageCoachWheelController.h"
+#include "Kismet2/BlueprintEditorUtils.h"
+#include "Kismet2/CompilerResultsLog.h"
+
+/////////////////////////////////////////////////////
+// UAnimGraphNode_StageCoachWheelController
+
+#define LOCTEXT_NAMESPACE "A3Nodes"
+
+UAnimGraphNode_StageCoachWheelController::UAnimGraphNode_StageCoachWheelController(const FObjectInitializer& ObjectInitializer)
+	: Super(ObjectInitializer)
+{
+}
+
+FText UAnimGraphNode_StageCoachWheelController::GetControllerDescription() const
+{
+	return LOCTEXT("AnimGraphNode_StagecoachEffectWheelController", "Stagecoach Effect Wheel Controller for ChaosWheeledVehicle");
+}
+
+FText UAnimGraphNode_StageCoachWheelController::GetTooltipText() const
+{
+	return LOCTEXT("AnimGraphNode_StagecoachEffectWheelController_Tooltip", "This alters the wheel transform based on set up in Wheeled Vehicle. This only works when the owner is ChaosWheeledVehicle.");
+}
+
+FText UAnimGraphNode_StageCoachWheelController::GetNodeTitle(ENodeTitleType::Type TitleType) const
+{
+	FText NodeTitle;
+	if (TitleType == ENodeTitleType::ListView || TitleType == ENodeTitleType::MenuTitle)
+	{
+		NodeTitle = GetControllerDescription();
+	}
+	else
+	{
+		// we don't have any run-time information, so it's limited to print  
+		// anymore than what it is it would be nice to print more data such as 
+		// name of bones for wheels, but it's not available in Persona
+		NodeTitle = FText(LOCTEXT("AnimGraphNode_StagecoachWheelController_Title", "Stagecoach Effect Wheel Controller"));
+	}	
+	return NodeTitle;
+}
+
+void UAnimGraphNode_StageCoachWheelController::ValidateAnimNodePostCompile(class FCompilerResultsLog& MessageLog, class UAnimBlueprintGeneratedClass* CompiledClass, int32 CompiledNodeIndex)
+{
+	// we only support vehicle anim instance
+	if ( CompiledClass->IsChildOf(UVehicleAnimationInstance::StaticClass())  == false )
+	{
+		MessageLog.Error(TEXT("@@ is only allowed in VehicleAnimInstance. If this is for vehicle, please change parent to be VehicleAnimInstance (Reparent Class)."), this);
+	}
+}
+
+bool UAnimGraphNode_StageCoachWheelController::IsCompatibleWithGraph(const UEdGraph* TargetGraph) const
+{
+	UBlueprint* Blueprint = FBlueprintEditorUtils::FindBlueprintForGraph(TargetGraph);
+	return (Blueprint != nullptr) && Blueprint->ParentClass->IsChildOf<UVehicleAnimationInstance>() && Super::IsCompatibleWithGraph(TargetGraph);
+}
+
+#undef LOCTEXT_NAMESPACE
diff --git a/Engine/Plugins/Experimental/ChaosVehiclesPlugin/Source/ChaosVehiclesEditor/Public/AnimGraphNode_StageCoachWheelController.h b/Engine/Plugins/Experimental/ChaosVehiclesPlugin/Source/ChaosVehiclesEditor/Public/AnimGraphNode_StageCoachWheelController.h
new file mode 100644
--- /dev/null
+++ b/Engine/Plugins/Experimental/ChaosVehiclesPlugin/Source/ChaosVehiclesEditor/Public/AnimGraphNode_StageCoachWheelController.h
@@ -0,0 +1,35 @@
+// Copyright Epic Games, Inc. All Rights Reserved.
+
+#pragma once
+
+#include "CoreMinimal.h"
+#include "UObject/ObjectMacros.h"
+#include "AnimGraphNode_SkeletalControlBase.h"
+#include "AnimNode_StageCoachWheelController.h"
+#include "AnimGraphNode_StageCoachWheelController.generated.h"
+
+class UEdGraph;
+
+UCLASS(MinimalAPI, meta=(Keywords = "Modify Wheel Vehicle"))
+class UAnimGraphNode_StageCoachWheelController : public UAnimGraphNode_SkeletalControlBase
+{
+	GENERATED_UCLASS_BODY()
+
+	UPROPERTY(EditAnywhere, Category=Settings)
+	FAnimNode_StageCoachWheelController Node;
+
+public:
+	// UEdGraphNode interface
+	virtual FText GetNodeTitle(ENodeTitleType::Type TitleType) const override;
+	virtual FText GetTooltipText() const override;
+	// validate if this is within VehicleAnimInstance
+	virtual void ValidateAnimNodePostCompile(class FCompilerResultsLog& MessageLog, class UAnimBlueprintGeneratedClass* CompiledClass, int32 CompiledNodeIndex) override;
+	virtual bool IsCompatibleWithGraph(const UEdGraph* TargetGraph) const override;
+	// End of UEdGraphNode interface
+
+protected:
+	// UAnimGraphNode_SkeletalControlBase interface
+	virtual FText GetControllerDescription() const override;
+	virtual const FAnimNode_SkeletalControlBase* GetNode() const override { return &Node; }
+	// End of UAnimGraphNode_SkeletalControlBase interface
+};
