Fix physics threading crash by moving apply forces out of parallel portion of code and performing them serially

diff --git a/Engine/Plugins/Experimental/ChaosVehiclesPlugin/Source/ChaosVehicles/Private/ChaosVehicleManagerAsyncCallback.cpp b/Engine/Plugins/Experimental/ChaosVehiclesPlugin/Source/ChaosVehicles/Private/ChaosVehicleManagerAsyncCallback.cpp
--- a/Engine/Plugins/Experimental/ChaosVehiclesPlugin/Source/ChaosVehicles/Private/ChaosVehicleManagerAsyncCallback.cpp
+++ b/Engine/Plugins/Experimental/ChaosVehiclesPlugin/Source/ChaosVehicles/Private/ChaosVehicleManagerAsyncCallback.cpp
@@ -48,8 +48,8 @@ void FChaosVehicleManagerAsyncCallback::OnPreSimulate_Internal()
 	Output.VehicleOutputs.AddDefaulted(NumVehicles);
 	Output.Timestamp = Input->Timestamp;
 
-	const auto& InputVehiclesBatch = Input->VehicleInputs;
-	auto& OutputVehiclesBatch = Output.VehicleOutputs;
+	const TArray<TUniquePtr<FChaosVehicleAsyncInput>>& InputVehiclesBatch = Input->VehicleInputs;
+	TArray<TUniquePtr<FChaosVehicleAsyncOutput>>& OutputVehiclesBatch = Output.VehicleOutputs;
 
 	// beware running the vehicle simulation in parallel, code must remain threadsafe
 	auto LambdaParallelUpdate = [World, DeltaTime, SimTime, &InputVehiclesBatch, &OutputVehiclesBatch](int32 Idx)
@@ -61,7 +61,7 @@ void FChaosVehicleManagerAsyncCallback::OnPreSimulate_Internal()
 			return;
 		}
 
-		auto Handle = VehicleInput.Proxy->GetPhysicsThreadAPI();
+		Chaos::FRigidBodyHandle_Internal* Handle = VehicleInput.Proxy->GetPhysicsThreadAPI();
 		if (Handle->ObjectState() != Chaos::EObjectStateType::Dynamic)
 		{
 			return;
@@ -75,7 +75,12 @@ void FChaosVehicleManagerAsyncCallback::OnPreSimulate_Internal()
 	bool ForceSingleThread = !GVehicleDebugParams.EnableMultithreading;
 	PhysicsParallelFor(OutputVehiclesBatch.Num(), LambdaParallelUpdate, ForceSingleThread);
 
-
+	// Delayed application of forces - This is seperate from Simulate because forces cannot be executed multi-threaded
+	for (const TUniquePtr<FChaosVehicleAsyncInput>& VehicleInput : InputVehiclesBatch)
+	{
+		Chaos::FRigidBodyHandle_Internal* Handle = VehicleInput->Proxy->GetPhysicsThreadAPI();
+		VehicleInput->ApplyDeferredForces(Handle);
+	}
 }
 
 /**
diff --git a/Engine/Plugins/Experimental/ChaosVehiclesPlugin/Source/ChaosVehicles/Private/ChaosVehicleMovementComponent.cpp b/Engine/Plugins/Experimental/ChaosVehiclesPlugin/Source/ChaosVehicles/Private/ChaosVehicleMovementComponent.cpp
--- a/Engine/Plugins/Experimental/ChaosVehiclesPlugin/Source/ChaosVehicles/Private/ChaosVehicleMovementComponent.cpp
+++ b/Engine/Plugins/Experimental/ChaosVehiclesPlugin/Source/ChaosVehicles/Private/ChaosVehicleMovementComponent.cpp
@@ -147,6 +147,11 @@ void UChaosVehicleSimulation::TickVehicle(UWorld* WorldIn, float DeltaTime, cons
 #endif
 }
 
+void UChaosVehicleSimulation::ApplyDeferredForces(Chaos::FRigidBodyHandle_Internal* Handle)
+{
+	DeferredForces.Apply(Handle);
+}
+
 void UChaosVehicleSimulation::UpdateSimulation(float DeltaTime, const FChaosVehicleDefaultAsyncInput& InputData, Chaos::FRigidBodyHandle_Internal* Handle)
 {
 	VehicleState.CaptureState(Handle, InputData.GravityZ, DeltaTime);
@@ -472,119 +477,65 @@ void UChaosVehicleSimulation::DrawDebug3D()
 #endif
 }
 
-void UChaosVehicleSimulation::AddForce(const FVector& Force, bool bAllowSubstepping /*= true*/, bool bAccelChange /*= false*/)
+void UChaosVehicleSimulation::AddForce(const FVector& Force, bool bAllowSubstepping, bool bAccelChange)
 {
-	if (ensure(RigidHandle))
-	{
-		Chaos::EObjectStateType ObjectState = RigidHandle->ObjectState();
-		if (CHAOS_ENSURE(ObjectState == Chaos::EObjectStateType::Dynamic || ObjectState == Chaos::EObjectStateType::Sleeping))
-		{
-			if (bAccelChange)
-			{
-				const float RigidMass = RigidHandle->M();
-				const Chaos::FVec3 Acceleration = Force * RigidMass;
-				RigidHandle->AddForce(Acceleration, false);
-			}
-			else
-			{
-				RigidHandle->AddForce(Force, false);
-			}
-
-		}
+	DeferredForces.Add(FDeferredForces::FApplyForceData(Force, bAllowSubstepping, bAccelChange));
 
 #if !(UE_BUILD_SHIPPING || UE_BUILD_TEST)
-		if (GVehicleDebugParams.ShowAllForces)
-		{
-			FVector Position = VehicleState.VehicleWorldCOM;
-			Chaos::FDebugDrawQueue::GetInstance().DrawDebugDirectionalArrow(Position, Position + Force * GVehicleDebugParams.ForceDebugScaling
-				, 20.f, FColor::Blue, false, 0, 0, 2.f);
-		}
-#endif	
+	if (GVehicleDebugParams.ShowAllForces)
+	{
+		FVector Position = RigidHandle->X();
+		Chaos::FDebugDrawQueue::GetInstance().DrawDebugDirectionalArrow(Position, Position + Force * GVehicleDebugParams.ForceDebugScaling
+			, 20.f, FColor::Blue, false, 0, 0, 2.f);
 	}
-
+#endif	
 }
 
-void UChaosVehicleSimulation::AddForceAtPosition(const FVector& Force, const FVector& Position, bool bAllowSubstepping /*= true*/, bool bIsLocalForce /*= false*/)
+void UChaosVehicleSimulation::AddForceAtPosition(const FVector& Force, const FVector& Position, bool bAllowSubstepping, bool bIsLocalForce)
 {
-	if (ensure(RigidHandle))
-	{
-		const Chaos::FVec3 WorldCOM = Chaos::FParticleUtilitiesGT::GetCoMWorldPosition(RigidHandle);
-		const Chaos::FVec3 WorldTorque = Chaos::FVec3::CrossProduct(Position - WorldCOM, Force);
-		RigidHandle->AddForce(Force, false);
-		RigidHandle->AddTorque(WorldTorque, false);
+	DeferredForces.Add(FDeferredForces::FApplyForceAtPositionData(Force, Position, bAllowSubstepping, bIsLocalForce));
 
 #if !(UE_BUILD_SHIPPING || UE_BUILD_TEST)
-		if (GVehicleDebugParams.ShowAllForces)
-		{
-			Chaos::FDebugDrawQueue::GetInstance().DrawDebugDirectionalArrow(Position, Position + Force * GVehicleDebugParams.ForceDebugScaling
-				, 20.f, FColor::Blue, false, 0, 0, 2.f);
-		}
-#endif
+	if (GVehicleDebugParams.ShowAllForces)
+	{
+		Chaos::FDebugDrawQueue::GetInstance().DrawDebugDirectionalArrow(Position, Position + Force * GVehicleDebugParams.ForceDebugScaling
+			, 20.f, FColor::Blue, false, 0, 0, 2.f);
 	}
+#endif
 }
 
 void UChaosVehicleSimulation::AddImpulse(const FVector& Impulse, bool bVelChange)
 {
-	if (ensure(RigidHandle))
-	{
-		if (bVelChange)
-		{
-			RigidHandle->SetLinearImpulse(RigidHandle->LinearImpulse() + RigidHandle->M() * Impulse, false);
-		}
-		else
-		{
-			RigidHandle->SetLinearImpulse(RigidHandle->LinearImpulse() + Impulse, false);
-		}
+	DeferredForces.Add(FDeferredForces::FAddImpulseData(Impulse, bVelChange));
 
 #if !(UE_BUILD_SHIPPING || UE_BUILD_TEST)
-		if (GVehicleDebugParams.ShowAllForces)
-		{
-			FVector Position = VehicleState.VehicleWorldCOM;
-			Chaos::FDebugDrawQueue::GetInstance().DrawDebugDirectionalArrow(Position, Position + Impulse * GVehicleDebugParams.ForceDebugScaling
-				, 20.f, FColor::Red, false, 0, 0, 2.f);
-		}
-#endif
+	if (GVehicleDebugParams.ShowAllForces)
+	{
+		FVector Position = VehicleState.VehicleWorldCOM;
+		Chaos::FDebugDrawQueue::GetInstance().DrawDebugDirectionalArrow(Position, Position + Impulse * GVehicleDebugParams.ForceDebugScaling
+			, 20.f, FColor::Red, false, 0, 0, 2.f);
 	}
+#endif
+	
 }
 
 void UChaosVehicleSimulation::AddImpulseAtPosition(const FVector& Impulse, const FVector& Position)
 {
-	if (ensure(RigidHandle))
-	{
-		const Chaos::FVec3 WorldCOM = Chaos::FParticleUtilitiesGT::GetCoMWorldPosition(RigidHandle);
-		const Chaos::FVec3 AngularImpulse = Chaos::FVec3::CrossProduct(Position - WorldCOM, Impulse);
-		RigidHandle->SetLinearImpulse(RigidHandle->LinearImpulse() + Impulse, false);
-		RigidHandle->SetAngularImpulse(RigidHandle->AngularImpulse() + AngularImpulse, false);
-
+	DeferredForces.Add(FDeferredForces::FAddImpulseAtPositionData(Impulse, Position));
 
 #if !(UE_BUILD_SHIPPING || UE_BUILD_TEST)
-		if (GVehicleDebugParams.ShowAllForces)
-		{
-			Chaos::FDebugDrawQueue::GetInstance().DrawDebugDirectionalArrow(Position, Position + Impulse * GVehicleDebugParams.ForceDebugScaling
-				, 20.f, FColor::Red, false, 0, 0, 2.f);
-		}
-#endif
+	if (GVehicleDebugParams.ShowAllForces)
+	{
+		Chaos::FDebugDrawQueue::GetInstance().DrawDebugDirectionalArrow(Position, Position + Impulse * GVehicleDebugParams.ForceDebugScaling
+			, 20.f, FColor::Red, false, 0, 0, 2.f);
 	}
+#endif
 
 }
 
 void UChaosVehicleSimulation::AddTorqueInRadians(const FVector& Torque, bool bAllowSubstepping /*= true*/, bool bAccelChange /*= false*/)
 {
-	if (ensure(RigidHandle))
-	{
-		Chaos::EObjectStateType ObjectState = RigidHandle->ObjectState();
-		if (CHAOS_ENSURE(ObjectState == Chaos::EObjectStateType::Dynamic || ObjectState == Chaos::EObjectStateType::Sleeping))
-		{
-			if (bAccelChange)
-			{
-				RigidHandle->AddTorque(Chaos::FParticleUtilitiesXR::GetWorldInertia(RigidHandle) * Torque, false);
-			}
-			else
-			{
-				RigidHandle->AddTorque(Torque, false);
-			}
-		}
-	}
+	DeferredForces.Add(FDeferredForces::FAddTorqueInRadiansData(Torque, bAllowSubstepping, bAccelChange));
 }
 
 void UChaosVehicleSimulation::InitializeWheel(int WheelIndex, const Chaos::FSimpleWheelConfig* InWheelSetup)
@@ -1706,6 +1657,14 @@ FChaosVehicleDefaultAsyncInput::FChaosVehicleDefaultAsyncInput()
 
 }
 
+
+void FChaosVehicleDefaultAsyncInput::ApplyDeferredForces(Chaos::FRigidBodyHandle_Internal* RigidHandle) const
+{
+	check(Vehicle);
+	check(Vehicle->VehicleSimulationPT);
+	Vehicle->VehicleSimulationPT->ApplyDeferredForces(RigidHandle);
+}
+
 /************************************************************************/
 /* Async simulation callback on the Physics Thread                      */
 /************************************************************************/
diff --git a/Engine/Plugins/Experimental/ChaosVehiclesPlugin/Source/ChaosVehicles/Private/DeferredForces.cpp b/Engine/Plugins/Experimental/ChaosVehiclesPlugin/Source/ChaosVehicles/Private/DeferredForces.cpp
new file mode 100644
--- /dev/null
+++ b/Engine/Plugins/Experimental/ChaosVehiclesPlugin/Source/ChaosVehicles/Private/DeferredForces.cpp
@@ -0,0 +1,89 @@
+// Copyright Epic Games, Inc. All Rights Reserved.
+
+#pragma once
+
+#include "DeferredForces.h"
+#include "PhysicsProxy/SingleParticlePhysicsProxy.h"
+#include "Chaos/Particle/ParticleUtilities.h"
+#include "Chaos/Utilities.h"
+
+
+void FDeferredForces::AddForce(Chaos::FRigidBodyHandle_Internal* RigidHandle, const FApplyForceData& DataIn)
+{
+	if (ensure(RigidHandle))
+	{
+		Chaos::EObjectStateType ObjectState = RigidHandle->ObjectState();
+		if (CHAOS_ENSURE(ObjectState == Chaos::EObjectStateType::Dynamic || ObjectState == Chaos::EObjectStateType::Sleeping))
+		{
+			if (DataIn.Flags == EForceFlags::AccelChange)
+			{
+				const float RigidMass = RigidHandle->M();
+				const Chaos::FVec3 Acceleration = DataIn.Force * RigidMass;
+				RigidHandle->AddForce(Acceleration, false);
+			}
+			else
+			{
+				RigidHandle->AddForce(DataIn.Force, false);
+			}
+
+		}
+	}
+
+}
+
+void FDeferredForces::AddForceAtPosition(Chaos::FRigidBodyHandle_Internal* RigidHandle, const FApplyForceAtPositionData& DataIn)
+{
+	if (ensure(RigidHandle))
+	{
+		const Chaos::FVec3 WorldCOM = Chaos::FParticleUtilitiesGT::GetCoMWorldPosition(RigidHandle);
+		const Chaos::FVec3 WorldTorque = Chaos::FVec3::CrossProduct(DataIn.Position - WorldCOM, DataIn.Force);
+		RigidHandle->AddForce(DataIn.Force, false);
+		RigidHandle->AddTorque(WorldTorque, false);
+	}
+}
+
+void FDeferredForces::AddTorque(Chaos::FRigidBodyHandle_Internal* RigidHandle, const FAddTorqueInRadiansData& DataIn)
+{
+	if (ensure(RigidHandle))
+	{
+		Chaos::EObjectStateType ObjectState = RigidHandle->ObjectState();
+		if (CHAOS_ENSURE(ObjectState == Chaos::EObjectStateType::Dynamic || ObjectState == Chaos::EObjectStateType::Sleeping))
+		{
+			if (DataIn.Flags == EForceFlags::AccelChange)
+			{
+				RigidHandle->AddTorque(Chaos::FParticleUtilitiesXR::GetWorldInertia(RigidHandle) * DataIn.Torque, false);
+			}
+			else
+			{
+				RigidHandle->AddTorque(DataIn.Torque, false);
+			}
+		}
+	}
+}
+
+void FDeferredForces::AddImpulse(Chaos::FRigidBodyHandle_Internal* RigidHandle, const FAddImpulseData& DataIn)
+{
+	if (ensure(RigidHandle))
+	{
+		if (DataIn.Flags == EForceFlags::VelChange)
+		{
+			RigidHandle->SetLinearImpulse(RigidHandle->LinearImpulse() + RigidHandle->M() * DataIn.Impulse, false);
+		}
+		else
+		{
+			RigidHandle->SetLinearImpulse(RigidHandle->LinearImpulse() + DataIn.Impulse, false);
+		}
+	}
+
+}
+
+void FDeferredForces::AddImpulseAtPosition(Chaos::FRigidBodyHandle_Internal* RigidHandle, const FAddImpulseAtPositionData& DataIn)
+{
+	if (ensure(RigidHandle))
+	{
+		const Chaos::FVec3 WorldCOM = Chaos::FParticleUtilitiesGT::GetCoMWorldPosition(RigidHandle);
+		const Chaos::FVec3 AngularImpulse = Chaos::FVec3::CrossProduct(DataIn.Position - WorldCOM, DataIn.Impulse);
+		RigidHandle->SetLinearImpulse(RigidHandle->LinearImpulse() + DataIn.Impulse, false);
+		RigidHandle->SetAngularImpulse(RigidHandle->AngularImpulse() + AngularImpulse, false);
+	}
+}
diff --git a/Engine/Plugins/Experimental/ChaosVehiclesPlugin/Source/ChaosVehicles/Public/ChaosVehicleManagerAsyncCallback.h b/Engine/Plugins/Experimental/ChaosVehiclesPlugin/Source/ChaosVehicles/Public/ChaosVehicleManagerAsyncCallback.h
--- a/Engine/Plugins/Experimental/ChaosVehiclesPlugin/Source/ChaosVehicles/Public/ChaosVehicleManagerAsyncCallback.h
+++ b/Engine/Plugins/Experimental/ChaosVehiclesPlugin/Source/ChaosVehicles/Public/ChaosVehicleManagerAsyncCallback.h
@@ -104,6 +104,8 @@ struct FChaosVehicleAsyncInput
 	*/
 	virtual TUniquePtr<struct FChaosVehicleAsyncOutput> Simulate(UWorld* World, const float DeltaSeconds, const float TotalSeconds, bool& bWakeOut) const = 0;
 
+	virtual void ApplyDeferredForces(Chaos::FRigidBodyHandle_Internal* RigidHandle) const = 0;
+
 	FChaosVehicleAsyncInput(EChaosAsyncVehicleDataType InType = EChaosAsyncVehicleDataType::AsyncInvalid)
 		: Type(InType)
 		, Vehicle(nullptr)
diff --git a/Engine/Plugins/Experimental/ChaosVehiclesPlugin/Source/ChaosVehicles/Public/ChaosVehicleMovementComponent.h b/Engine/Plugins/Experimental/ChaosVehiclesPlugin/Source/ChaosVehicles/Public/ChaosVehicleMovementComponent.h
--- a/Engine/Plugins/Experimental/ChaosVehiclesPlugin/Source/ChaosVehicles/Public/ChaosVehicleMovementComponent.h
+++ b/Engine/Plugins/Experimental/ChaosVehiclesPlugin/Source/ChaosVehicles/Public/ChaosVehicleMovementComponent.h
@@ -14,6 +14,7 @@
 #include "ThrustSystem.h"
 #include "PhysicsProxy/SingleParticlePhysicsProxyFwd.h"
 #include "SnapshotData.h"
+#include "DeferredForces.h"
 #include "ChaosVehicleManagerAsyncCallback.h"
 
 #include "ChaosVehicleMovementComponent.generated.h"
@@ -171,7 +172,10 @@ struct FChaosVehicleDefaultAsyncInput : public FChaosVehicleAsyncInput
 
 	FChaosVehicleDefaultAsyncInput();
 
-	TUniquePtr<FChaosVehicleAsyncOutput> Simulate(UWorld* World, const float DeltaSeconds, const float TotalSeconds, bool& bWakeOut) const override;
+	virtual TUniquePtr<FChaosVehicleAsyncOutput> Simulate(UWorld* World, const float DeltaSeconds, const float TotalSeconds, bool& bWakeOut) const override;
+	
+	virtual void ApplyDeferredForces(Chaos::FRigidBodyHandle_Internal* RigidHandle) const override;
+
 };
 
 
@@ -728,6 +732,8 @@ public:
 
 	virtual void TickVehicle(UWorld* WorldIn, float DeltaTime, const FChaosVehicleDefaultAsyncInput& InputData, FChaosVehicleAsyncOutput& OutputData, Chaos::FRigidBodyHandle_Internal* Handle);
 
+	virtual void ApplyDeferredForces(Chaos::FRigidBodyHandle_Internal* Handle);
+
 	/** Advance the vehicle simulation */
 	virtual void UpdateSimulation(float DeltaTime, const FChaosVehicleDefaultAsyncInput& InputData, Chaos::FRigidBodyHandle_Internal* Handle);
 	
@@ -781,6 +787,8 @@ public:
 	// #todo: this isn't very configurable
 	TUniquePtr<Chaos::FSimpleWheeledVehicle> PVehicle;
 
+	FDeferredForces DeferredForces;
+
 };
 
 
diff --git a/Engine/Plugins/Experimental/ChaosVehiclesPlugin/Source/ChaosVehicles/Public/DeferredForces.h b/Engine/Plugins/Experimental/ChaosVehiclesPlugin/Source/ChaosVehicles/Public/DeferredForces.h
new file mode 100644
--- /dev/null
+++ b/Engine/Plugins/Experimental/ChaosVehiclesPlugin/Source/ChaosVehicles/Public/DeferredForces.h
@@ -0,0 +1,170 @@
+// Copyright Epic Games, Inc. All Rights Reserved.
+
+#pragma once
+
+#include "CoreMinimal.h"
+
+
+enum class EForceFlags : uint32
+{
+	None = 0, // No flags.
+
+	AllowSubstepping = 1 << 0,
+	AccelChange = 1 << 1,
+	VelChange = 1 << 2,
+	IsLocalForce = 1 << 3
+};
+ENUM_CLASS_FLAGS(EForceFlags);	
+
+class FDeferredForces
+{
+public:
+
+
+
+	struct FApplyForceData
+	{
+
+
+		FApplyForceData(const FVector& ForceIn, bool bAllowSubsteppingIn, bool bAccelChangeIn)
+			: Force(ForceIn)
+			, Flags(EForceFlags::None)
+		{
+			Flags |= bAllowSubsteppingIn ? EForceFlags::AllowSubstepping : EForceFlags::None;
+			Flags |= bAccelChangeIn ? EForceFlags::AccelChange : EForceFlags::None;
+		}
+
+		FVector Force;
+		EForceFlags Flags;
+	};
+
+	struct FApplyForceAtPositionData
+	{
+		FApplyForceAtPositionData(const FVector& ForceIn, const FVector& PositionIn, bool bAllowSubsteppingIn, bool bIsLocalForceIn)
+			: Force(ForceIn)
+			, Position(PositionIn)
+			, Flags(EForceFlags::None)
+		{
+			Flags |= bAllowSubsteppingIn ? EForceFlags::AllowSubstepping : EForceFlags::None;
+			Flags |= bIsLocalForceIn ? EForceFlags::IsLocalForce : EForceFlags::None;
+		}
+
+		FVector Force;
+		FVector Position;
+		EForceFlags Flags;
+	};
+
+	struct FAddTorqueInRadiansData
+	{
+		FAddTorqueInRadiansData(const FVector& TorqueIn, bool bAllowSubsteppingIn, bool bAccelChangeIn)
+			: Torque(TorqueIn)
+			, Flags(EForceFlags::None)
+		{
+			Flags |= bAllowSubsteppingIn ? EForceFlags::AllowSubstepping : EForceFlags::None;
+			Flags |= bAccelChangeIn ? EForceFlags::AccelChange : EForceFlags::None;
+		}
+
+		FVector Torque;
+		EForceFlags Flags;
+	};
+
+	struct FAddImpulseData
+	{
+		FAddImpulseData(const FVector& ImpulseIn, const bool bVelChangeIn)
+			: Impulse(ImpulseIn)
+			, Flags(EForceFlags::None)
+		{
+			Flags |= bVelChangeIn ? EForceFlags::VelChange : EForceFlags::None;
+		}
+
+		FVector Impulse;
+		EForceFlags Flags;
+	};
+
+	struct FAddImpulseAtPositionData
+	{
+		FAddImpulseAtPositionData(const FVector& ImpulseIn, const FVector& PositionIn)
+			: Impulse(ImpulseIn)
+			, Position(PositionIn)
+		{
+
+		}
+
+		FVector Impulse;
+		FVector Position;
+	};
+
+	void Add(const FApplyForceData& ApplyForceDataIn)
+	{
+		ApplyForceDatas.Add(ApplyForceDataIn);
+	}
+
+	void Add(const FApplyForceAtPositionData& ApplyForceAtPositionDataIn)
+	{
+		ApplyForceAtPositionDatas.Add(ApplyForceAtPositionDataIn);
+	}
+
+	void Add(const FAddTorqueInRadiansData& ApplyTorqueDataIn)
+	{
+		ApplyTorqueDatas.Add(ApplyTorqueDataIn);
+	}
+
+	void Add(const FAddImpulseData& ApplyImpulseDataIn)
+	{
+		ApplyImpulseDatas.Add(ApplyImpulseDataIn);
+	}
+
+	void Add(const FAddImpulseAtPositionData& ApplyImpulseAtPositionDataIn)
+	{
+		ApplyImpulseAtPositionDatas.Add(ApplyImpulseAtPositionDataIn);
+	}
+
+
+	void Apply(Chaos::FRigidBodyHandle_Internal* RigidHandle)
+	{
+		for (const FApplyForceData& Data : ApplyForceDatas)
+		{
+			AddForce(RigidHandle, Data);
+		}
+
+		for (const FApplyForceAtPositionData& Data : ApplyForceAtPositionDatas)
+		{
+			AddForceAtPosition(RigidHandle, Data);
+		}
+
+		for (const FAddTorqueInRadiansData& Data : ApplyTorqueDatas)
+		{
+			AddTorque(RigidHandle, Data);
+		}
+
+		for (const FAddImpulseData& Data : ApplyImpulseDatas)
+		{
+			AddImpulse(RigidHandle, Data);
+		}
+
+		for (const FAddImpulseAtPositionData& Data : ApplyImpulseAtPositionDatas)
+		{
+			AddImpulseAtPosition(RigidHandle, Data);
+		}
+
+		ApplyForceDatas.Empty();
+		ApplyForceAtPositionDatas.Empty();
+		ApplyTorqueDatas.Empty();
+		ApplyImpulseDatas.Empty();
+		ApplyImpulseAtPositionDatas.Empty();
+	}
+
+private:
+
+	void AddForce(Chaos::FRigidBodyHandle_Internal* RigidHandle, const FApplyForceData& DataIn);
+	void AddForceAtPosition(Chaos::FRigidBodyHandle_Internal* RigidHandle, const FApplyForceAtPositionData& DataIn);
+	void AddTorque(Chaos::FRigidBodyHandle_Internal* RigidHandle, const FAddTorqueInRadiansData& DataIn);
+	void AddImpulse(Chaos::FRigidBodyHandle_Internal* RigidHandle, const FAddImpulseData& DataIn);
+	void AddImpulseAtPosition(Chaos::FRigidBodyHandle_Internal* RigidHandle, const FAddImpulseAtPositionData& DataIn);
+
+	TArray<FApplyForceData> ApplyForceDatas;
+	TArray<FApplyForceAtPositionData> ApplyForceAtPositionDatas;
+	TArray<FAddTorqueInRadiansData> ApplyTorqueDatas;
+	TArray<FAddImpulseData> ApplyImpulseDatas;
+	TArray<FAddImpulseAtPositionData> ApplyImpulseAtPositionDatas;
+};
