Fix CIS non-unity build

diff --git a/Engine/Plugins/Experimental/ChaosVehiclesPlugin/Source/ChaosVehicles/Public/DeferredForces.h b/Engine/Plugins/Experimental/ChaosVehiclesPlugin/Source/ChaosVehicles/Public/DeferredForces.h
--- a/Engine/Plugins/Experimental/ChaosVehiclesPlugin/Source/ChaosVehicles/Public/DeferredForces.h
+++ b/Engine/Plugins/Experimental/ChaosVehiclesPlugin/Source/ChaosVehicles/Public/DeferredForces.h
@@ -16,6 +16,14 @@ enum class EForceFlags : uint32
 };
 ENUM_CLASS_FLAGS(EForceFlags);	
 
+
+// Forward declarations
+namespace Chaos
+{ 
+	class FRigidBodyHandle_Internal;
+}
+
+
 class FDeferredForces
 {
 public:
