Fix crash check for null pointers.

diff --git a/Engine/Plugins/Experimental/ChaosVehiclesPlugin/Source/ChaosVehicles/Private/ChaosVehicleManagerAsyncCallback.cpp b/Engine/Plugins/Experimental/ChaosVehiclesPlugin/Source/ChaosVehicles/Private/ChaosVehicleManagerAsyncCallback.cpp
--- a/Engine/Plugins/Experimental/ChaosVehiclesPlugin/Source/ChaosVehicles/Private/ChaosVehicleManagerAsyncCallback.cpp
+++ b/Engine/Plugins/Experimental/ChaosVehiclesPlugin/Source/ChaosVehicles/Private/ChaosVehicleManagerAsyncCallback.cpp
@@ -75,11 +75,16 @@ void FChaosVehicleManagerAsyncCallback::OnPreSimulate_Internal()
 	bool ForceSingleThread = !GVehicleDebugParams.EnableMultithreading;
 	PhysicsParallelFor(OutputVehiclesBatch.Num(), LambdaParallelUpdate, ForceSingleThread);
 
-	// Delayed application of forces - This is seperate from Simulate because forces cannot be executed multi-threaded
+	// Delayed application of forces - This is separate from Simulate because forces cannot be executed multi-threaded
 	for (const TUniquePtr<FChaosVehicleAsyncInput>& VehicleInput : InputVehiclesBatch)
 	{
-		Chaos::FRigidBodyHandle_Internal* Handle = VehicleInput->Proxy->GetPhysicsThreadAPI();
-		VehicleInput->ApplyDeferredForces(Handle);
+		if (VehicleInput.IsValid() && VehicleInput->Proxy)
+		{
+			if (Chaos::FRigidBodyHandle_Internal* Handle = VehicleInput->Proxy->GetPhysicsThreadAPI())
+			{
+				VehicleInput->ApplyDeferredForces(Handle);
+			}
+		}
 	}
 }
 
