Bug fix vehicle trailer crash

diff --git a/Engine/Plugins/Experimental/ChaosVehiclesPlugin/Source/ChaosVehicles/Private/ChaosWheeledVehicleMovementComponent.cpp b/Engine/Plugins/Experimental/ChaosVehiclesPlugin/Source/ChaosVehicles/Private/ChaosWheeledVehicleMovementComponent.cpp
--- a/Engine/Plugins/Experimental/ChaosVehiclesPlugin/Source/ChaosVehicles/Private/ChaosWheeledVehicleMovementComponent.cpp
+++ b/Engine/Plugins/Experimental/ChaosVehiclesPlugin/Source/ChaosVehicles/Private/ChaosWheeledVehicleMovementComponent.cpp
@@ -1399,6 +1399,18 @@ void UChaosWheeledVehicleMovementComponent::SetupVehicle(TUniquePtr<Chaos::FSimp
 	// cache this value as it's useful for steering setup calculations and debug rendering
 	WheelTrackDimensions = CalculateWheelLayoutDimensions();
 
+	if (EngineSetup.TorqueCurve.GetRichCurve()->IsEmpty())
+	{
+		FString ActorName = "Unknown";
+		if (GetOwner())
+		{
+			ActorName = GetOwner()->GetName();
+		}
+		UE_LOG(LogVehicle, Warning, TEXT("Vehicle %s has no torque curve defined, disabling mechanical simulation."), *ActorName);
+
+		bMechanicalSimEnabled = false;
+	}
+
 	if (bMechanicalSimEnabled)
 	{
 		Chaos::FSimpleEngineSim EngineSim(&EngineSetup.GetPhysicsEngineConfig());
