Fix issue where axle setup override wasn't functioning quite right.

diff --git a/Engine/Plugins/Experimental/ChaosVehiclesPlugin/Source/ChaosVehicles/Private/ChaosWheeledVehicleMovementComponent.cpp b/Engine/Plugins/Experimental/ChaosVehiclesPlugin/Source/ChaosVehicles/Private/ChaosWheeledVehicleMovementComponent.cpp
--- a/Engine/Plugins/Experimental/ChaosVehiclesPlugin/Source/ChaosVehicles/Private/ChaosWheeledVehicleMovementComponent.cpp
+++ b/Engine/Plugins/Experimental/ChaosVehiclesPlugin/Source/ChaosVehicles/Private/ChaosWheeledVehicleMovementComponent.cpp
@@ -1429,6 +1429,7 @@ void UChaosWheeledVehicleMovementComponent::SetupVehicle(TUniquePtr<Chaos::FSimp
 			FSimpleWheelSim& PWheel = PVehicle->Wheels[WheelIdx];
 			bool IsWheelPowered = FTransmissionUtility::IsWheelPowered(DifferentialSim.Setup().DifferentialType, PWheel.Setup().AxleType, PWheel.EngineEnabled);
 			PWheel.AccessSetup().EngineEnabled = IsWheelPowered;
+			PWheel.EngineEnabled = IsWheelPowered;
 
 			float TorqueRatio = FTransmissionUtility::GetTorqueRatioForWheel(DifferentialSim, WheelIdx, PVehicle->Wheels);
 			PWheel.AccessSetup().TorqueRatio = TorqueRatio;
