Fix deferred forces flags

diff --git a/Engine/Plugins/Experimental/ChaosVehiclesPlugin/Source/ChaosVehicles/Private/DeferredForces.cpp b/Engine/Plugins/Experimental/ChaosVehiclesPlugin/Source/ChaosVehicles/Private/DeferredForces.cpp
--- a/Engine/Plugins/Experimental/ChaosVehiclesPlugin/Source/ChaosVehicles/Private/DeferredForces.cpp
+++ b/Engine/Plugins/Experimental/ChaosVehiclesPlugin/Source/ChaosVehicles/Private/DeferredForces.cpp
@@ -13,7 +13,7 @@ void FDeferredForces::AddForce(Chaos::FRigidBodyHandle_Internal* RigidHandle, co
 		Chaos::EObjectStateType ObjectState = RigidHandle->ObjectState();
 		if (CHAOS_ENSURE(ObjectState == Chaos::EObjectStateType::Dynamic || ObjectState == Chaos::EObjectStateType::Sleeping))
 		{
-			if (DataIn.Flags == EForceFlags::AccelChange)
+			if ((DataIn.Flags & EForceFlags::AccelChange) == EForceFlags::AccelChange)
 			{
 				const float RigidMass = RigidHandle->M();
 				const Chaos::FVec3 Acceleration = DataIn.Force * RigidMass;
@@ -47,7 +47,7 @@ void FDeferredForces::AddTorque(Chaos::FRigidBodyHandle_Internal* RigidHandle, c
 		Chaos::EObjectStateType ObjectState = RigidHandle->ObjectState();
 		if (CHAOS_ENSURE(ObjectState == Chaos::EObjectStateType::Dynamic || ObjectState == Chaos::EObjectStateType::Sleeping))
 		{
-			if (DataIn.Flags == EForceFlags::AccelChange)
+			if ((DataIn.Flags & EForceFlags::AccelChange) == EForceFlags::AccelChange)
 			{
 				RigidHandle->AddTorque(Chaos::FParticleUtilitiesXR::GetWorldInertia(RigidHandle) * DataIn.Torque, false);
 			}
@@ -63,7 +63,7 @@ void FDeferredForces::AddImpulse(Chaos::FRigidBodyHandle_Internal* RigidHandle,
 {
 	if (ensure(RigidHandle))
 	{
-		if (DataIn.Flags == EForceFlags::VelChange)
+		if ((DataIn.Flags & EForceFlags::VelChange) == EForceFlags::VelChange)
 		{
 			RigidHandle->SetLinearImpulse(RigidHandle->LinearImpulse() + RigidHandle->M() * DataIn.Impulse, false);
 		}
