Updated comment to better describe a vehicle tuning property

diff --git a/Engine/Plugins/Experimental/ChaosVehiclesPlugin/Source/ChaosVehicles/Public/ChaosVehicleWheel.h b/Engine/Plugins/Experimental/ChaosVehiclesPlugin/Source/ChaosVehicles/Public/ChaosVehicleWheel.h
--- a/Engine/Plugins/Experimental/ChaosVehiclesPlugin/Source/ChaosVehicles/Public/ChaosVehicleWheel.h
+++ b/Engine/Plugins/Experimental/ChaosVehiclesPlugin/Source/ChaosVehicles/Public/ChaosVehicleWheel.h
@@ -155,7 +155,10 @@ class UChaosWheeledVehicleMovementComponent;
 		float SuspensionDampingRatio;
 
 		/**
-		 *	When 0 no weight load is transferred, 1 is Normal weight shift. Lower value cures lift off oversteer.
+		 *	Amount wheel load effects wheel friction. 
+			At 0 wheel friction is completely independent of the loading on the wheel (This is artificial as it always assumes even balance between all wheels)
+			At 1 wheel friction is based on the force pressing wheel into the ground. This is more realistic.
+			Lower value cures lift off over-steer, generally makes vehicle easier to handle under extreme motions.
 		 */
 		UPROPERTY(EditAnywhere, Category = Suspension, meta = (ClampMin = "0.0", UIMin = "0.0", ClampMax = "1.0", UIMax = "1.0"))
 		float WheelLoadRatio;
