Fix wheels not rotating when vehicle suspended on box

#okforgihub public, internal

diff --git a/Engine/Plugins/Experimental/ChaosVehiclesPlugin/Source/ChaosVehicles/Private/ChaosWheeledVehicleMovementComponent.cpp b/Engine/Plugins/Experimental/ChaosVehiclesPlugin/Source/ChaosVehicles/Private/ChaosWheeledVehicleMovementComponent.cpp
--- a/Engine/Plugins/Experimental/ChaosVehiclesPlugin/Source/ChaosVehicles/Private/ChaosWheeledVehicleMovementComponent.cpp
+++ b/Engine/Plugins/Experimental/ChaosVehiclesPlugin/Source/ChaosVehicles/Private/ChaosWheeledVehicleMovementComponent.cpp
@@ -558,7 +558,8 @@ void UChaosWheeledVehicleSimulation::ApplyWheelFrictionForces(float DeltaTime)
 		}
 		else
 		{
-			PWheel.SetVehicleGroundSpeed(WheelState.LocalWheelVelocity[WheelIdx]);
+			PWheel.SetVehicleGroundSpeed(FVector::ZeroVector);
+			PWheel.SetWheelLoadForce(0.f);
 			PWheel.Simulate(DeltaTime);
 		}
 
@@ -651,6 +652,8 @@ void UChaosWheeledVehicleSimulation::ApplySuspensionForces(float DeltaTime)
 		else
 		{
 			PSuspension.SetSuspensionLength(PSuspension.GetTraceLength(PWheel.GetEffectiveRadius()), PWheel.Setup().WheelRadius);
+			PWheel.SetWheelLoadForce(0.f);
+
 		}
 
 	}
diff --git a/Engine/Source/Runtime/Experimental/ChaosVehicles/ChaosVehiclesCore/Private/WheelSystem.cpp b/Engine/Source/Runtime/Experimental/ChaosVehicles/ChaosVehiclesCore/Private/WheelSystem.cpp
--- a/Engine/Source/Runtime/Experimental/ChaosVehicles/ChaosVehiclesCore/Private/WheelSystem.cpp
+++ b/Engine/Source/Runtime/Experimental/ChaosVehicles/ChaosVehiclesCore/Private/WheelSystem.cpp
@@ -24,6 +24,7 @@ namespace Chaos
 		, Re(SetupIn->WheelRadius)
 		, Omega(0.f)
 		, Sx(0.f)
+		, Inertia(0.5f * SetupIn->WheelMass * SetupIn->WheelRadius * SetupIn->WheelRadius)
 		, DriveTorque(0.f)
 		, BrakeTorque(0.f)
 		, ForceIntoSurface(0.f)
@@ -169,6 +170,12 @@ namespace Chaos
 				}
 			}
 		}
+		else
+		{
+			SlipOmega = DriveTorque / Inertia * DeltaTime;
+			SlipOmega = FMath::Clamp(SlipOmega, -Setup().MaxSpinRotation, Setup().MaxSpinRotation);
+		}
+
 
 		if (WheelLocked)
 		{
diff --git a/Engine/Source/Runtime/Experimental/ChaosVehicles/ChaosVehiclesCore/Public/WheelSystem.h b/Engine/Source/Runtime/Experimental/ChaosVehicles/ChaosVehiclesCore/Public/WheelSystem.h
--- a/Engine/Source/Runtime/Experimental/ChaosVehicles/ChaosVehiclesCore/Public/WheelSystem.h
+++ b/Engine/Source/Runtime/Experimental/ChaosVehicles/ChaosVehiclesCore/Public/WheelSystem.h
@@ -191,7 +191,7 @@ public:
 	}
 
 	/** Set the vehicle's speed at the wheels location in local wheel coords */
-	void SetVehicleGroundSpeed(FVector& VIn)
+	void SetVehicleGroundSpeed(const FVector& VIn)
 	{
 		GroundVelocityVector = VIn;
 	}
@@ -402,6 +402,7 @@ public:
 	float Re;		// [cm] Effective Wheel Radius could change dynamically if get a flat?, tire shreds
 	float Omega;	// [radians/sec] Wheel Rotation Angular Velocity
 	float Sx;
+	float Inertia;
 	// In
 	float DriveTorque;				// [N.m]
 	float BrakeTorque;				// [N.m]
