Change the vehicle wheel animation data use the correct suspension axis rather than an up vector.

diff --git a/Engine/Plugins/Experimental/ChaosVehiclesPlugin/Source/ChaosVehicles/Private/ChaosVehicleWheel.cpp b/Engine/Plugins/Experimental/ChaosVehiclesPlugin/Source/ChaosVehicles/Private/ChaosVehicleWheel.cpp
--- a/Engine/Plugins/Experimental/ChaosVehiclesPlugin/Source/ChaosVehicles/Private/ChaosVehicleWheel.cpp
+++ b/Engine/Plugins/Experimental/ChaosVehiclesPlugin/Source/ChaosVehicles/Private/ChaosVehicleWheel.cpp
@@ -99,6 +99,12 @@ float UChaosVehicleWheel::GetSuspensionOffset() const
 	return VehicleComponent->GetSuspensionOffset(WheelIndex);
 }
 
+FVector UChaosVehicleWheel::GetSuspensionAxis() const
+{
+	check(VehicleComponent);
+	return SuspensionAxis;
+}
+
 bool UChaosVehicleWheel::IsInAir() const
 {
 	check(VehicleComponent && VehicleComponent->PhysicsVehicleOutput());
diff --git a/Engine/Plugins/Experimental/ChaosVehiclesPlugin/Source/ChaosVehicles/Private/VehicleAnimationInstance.cpp b/Engine/Plugins/Experimental/ChaosVehiclesPlugin/Source/ChaosVehicles/Private/VehicleAnimationInstance.cpp
--- a/Engine/Plugins/Experimental/ChaosVehiclesPlugin/Source/ChaosVehicles/Private/VehicleAnimationInstance.cpp
+++ b/Engine/Plugins/Experimental/ChaosVehiclesPlugin/Source/ChaosVehicles/Private/VehicleAnimationInstance.cpp
@@ -128,9 +128,7 @@
 						WheelInstance.RotOffset.Yaw = VehicleWheel->GetSteerAngle();
 						WheelInstance.RotOffset.Roll = 0.f;
 
-						WheelInstance.LocOffset.X = 0.f;
-						WheelInstance.LocOffset.Y = 0.f;
-						WheelInstance.LocOffset.Z = VehicleWheel->GetSuspensionOffset();
+						WheelInstance.LocOffset = VehicleWheel->GetSuspensionAxis() * VehicleWheel->GetSuspensionOffset();
 					}
 				}
 			}
diff --git a/Engine/Plugins/Experimental/ChaosVehiclesPlugin/Source/ChaosVehicles/Public/ChaosVehicleWheel.h b/Engine/Plugins/Experimental/ChaosVehiclesPlugin/Source/ChaosVehicles/Public/ChaosVehicleWheel.h
--- a/Engine/Plugins/Experimental/ChaosVehiclesPlugin/Source/ChaosVehicles/Public/ChaosVehicleWheel.h
+++ b/Engine/Plugins/Experimental/ChaosVehiclesPlugin/Source/ChaosVehicles/Public/ChaosVehicleWheel.h
@@ -263,6 +263,9 @@ class UChaosWheeledVehicleMovementComponent;
 		UFUNCTION(BlueprintCallable, Category = "Game|Components|WheeledVehicleMovement")
 		float GetWheelAngularVelocity() const;
 
+		UFUNCTION(BlueprintCallable, Category = "Game|Components|WheeledVehicleMovement")
+		FVector GetSuspensionAxis() const;
+
 		UFUNCTION(BlueprintCallable, Category = "Game|Components|WheeledVehicleMovement")
 		bool IsInAir() const;
 
