Chaos, CIS: Fix NoPCH and NoUnity builds due to physics files.

diff --git a/Engine/Plugins/Experimental/ChaosCaching/Source/ChaosCaching/Private/Chaos/CacheManagerActor.cpp b/Engine/Plugins/Experimental/ChaosCaching/Source/ChaosCaching/Private/Chaos/CacheManagerActor.cpp
--- a/Engine/Plugins/Experimental/ChaosCaching/Source/ChaosCaching/Private/Chaos/CacheManagerActor.cpp
+++ b/Engine/Plugins/Experimental/ChaosCaching/Source/ChaosCaching/Private/Chaos/CacheManagerActor.cpp
@@ -9,6 +9,8 @@
 #include "ChaosSolversModule.h"
 #include "Components/BillboardComponent.h"
 #include "PBDRigidsSolver.h"
+#include "Features/IModularFeatures.h"
+#include "Algo/Find.h"
 
 #if WITH_EDITOR
 #include "Editor.h"
diff --git a/Engine/Plugins/Experimental/ChaosVehiclesPlugin/Source/ChaosVehicles/Private/ChaosVehicleManager.cpp b/Engine/Plugins/Experimental/ChaosVehiclesPlugin/Source/ChaosVehicles/Private/ChaosVehicleManager.cpp
--- a/Engine/Plugins/Experimental/ChaosVehiclesPlugin/Source/ChaosVehicles/Private/ChaosVehicleManager.cpp
+++ b/Engine/Plugins/Experimental/ChaosVehiclesPlugin/Source/ChaosVehicles/Private/ChaosVehicleManager.cpp
@@ -31,6 +31,24 @@ FDelegateHandle FChaosVehicleManager::OnWorldCleanupHandle;
 bool FChaosVehicleManager::GInitialized = false;
 #endif
 
+void OnPostWorldInitialize(UWorld* InWorld, const UWorld::InitializationValues)
+{
+	FChaosVehicleManager* Manager = FChaosVehicleManager::GetVehicleManagerFromScene(InWorld->GetPhysicsScene());
+	if(Manager)
+	{
+		Manager->RegisterCallbacks();
+	}
+}
+
+void OnWorldCleanup(UWorld* InWorld, bool bSessionEnded, bool bCleanupResources)
+{
+	FChaosVehicleManager* Manager = FChaosVehicleManager::GetVehicleManagerFromScene(InWorld->GetPhysicsScene());
+	if(Manager)
+	{
+		Manager->UnregisterCallbacks();
+	}
+}
+
 FChaosVehicleManager::FChaosVehicleManager(FPhysScene* PhysScene)
 #if WITH_CHAOS
 	: Scene(*PhysScene)
@@ -48,8 +66,8 @@ FChaosVehicleManager::FChaosVehicleManager(FPhysScene* PhysScene)
 		// PhysScene->GetOwningWorld() is always null here, the world is being setup too late to be of use
 		// therefore setup these global world delegates that will callback when everything is setup so registering
 		// the physics solver Async Callback will succeed
-		OnPostWorldInitializationHandle = FWorldDelegates::OnPostWorldInitialization.AddStatic(&FChaosVehicleManager::OnPostWorldInitialization);
-		OnWorldCleanupHandle = FWorldDelegates::OnWorldCleanup.AddStatic(&FChaosVehicleManager::OnWorldCleanup);
+		OnPostWorldInitializationHandle = FWorldDelegates::OnPostWorldInitialization.AddStatic(&OnPostWorldInitialize);
+		OnWorldCleanupHandle = FWorldDelegates::OnWorldCleanup.AddStatic(&OnWorldCleanup);
 	}
 
 	ensure(FChaosVehicleManager::SceneToVehicleManagerMap.Find(PhysScene) == nullptr);	//double registration with same scene, will cause a leak
@@ -85,25 +103,6 @@ void FChaosVehicleManager::UnregisterCallbacks()
 #endif
 }
 
-void FChaosVehicleManager::OnPostWorldInitialization(UWorld* InWorld, const UWorld::InitializationValues)
-{
-	FChaosVehicleManager* Manager = FChaosVehicleManager::GetVehicleManagerFromScene(InWorld->GetPhysicsScene());
-	if (Manager)
-	{
-		Manager->RegisterCallbacks();
-	}
-}
-
-
-void FChaosVehicleManager::OnWorldCleanup(UWorld* InWorld, bool bSessionEnded, bool bCleanupResources)
-{
-	FChaosVehicleManager* Manager = FChaosVehicleManager::GetVehicleManagerFromScene(InWorld->GetPhysicsScene());
-	if (Manager)
-	{
-		Manager->UnregisterCallbacks();
-	}
-}
-
 void FChaosVehicleManager::DetachFromPhysScene(FPhysScene* PhysScene)
 {
 	if (AsyncCallback)
diff --git a/Engine/Plugins/Experimental/ChaosVehiclesPlugin/Source/ChaosVehicles/Public/ChaosVehicleManager.h b/Engine/Plugins/Experimental/ChaosVehiclesPlugin/Source/ChaosVehicles/Public/ChaosVehicleManager.h
--- a/Engine/Plugins/Experimental/ChaosVehiclesPlugin/Source/ChaosVehicles/Public/ChaosVehicleManager.h
+++ b/Engine/Plugins/Experimental/ChaosVehiclesPlugin/Source/ChaosVehicles/Public/ChaosVehicleManager.h
@@ -12,7 +12,7 @@
 
 class UChaosTireConfig;
 class UChaosVehicleMovementComponent;
-
+class FChaosScene;
 
 class CHAOSVEHICLES_API FChaosVehicleManager
 {
@@ -24,9 +24,6 @@ public:
 	FChaosVehicleManager(FPhysScene* PhysScene);
 	~FChaosVehicleManager();
 
-	static void OnPostWorldInitialization(UWorld* InWorld, const UWorld::InitializationValues);
-	static void OnWorldCleanup(UWorld* InWorld, bool bSessionEnded, bool bCleanupResources);
-
 #if WITH_CHAOS
 	/** Get Physics Scene */
 	FPhysScene_Chaos& GetScene() const { return Scene; }
diff --git a/Engine/Plugins/Experimental/ChaosVehiclesPlugin/Source/ChaosVehicles/Public/ChaosVehicleMovementComponent.h b/Engine/Plugins/Experimental/ChaosVehiclesPlugin/Source/ChaosVehicles/Public/ChaosVehicleMovementComponent.h
--- a/Engine/Plugins/Experimental/ChaosVehiclesPlugin/Source/ChaosVehicles/Public/ChaosVehicleMovementComponent.h
+++ b/Engine/Plugins/Experimental/ChaosVehiclesPlugin/Source/ChaosVehicles/Public/ChaosVehicleMovementComponent.h
@@ -16,6 +16,7 @@
 #include "SnapshotData.h"
 #include "DeferredForces.h"
 #include "ChaosVehicleManagerAsyncCallback.h"
+#include "CollisionQueryParams.h"
 
 #include "ChaosVehicleMovementComponent.generated.h"
 
diff --git a/Engine/Source/Runtime/Engine/Private/Collision/PhysicsFiltering.cpp b/Engine/Source/Runtime/Engine/Private/Collision/PhysicsFiltering.cpp
--- a/Engine/Source/Runtime/Engine/Private/Collision/PhysicsFiltering.cpp
+++ b/Engine/Source/Runtime/Engine/Private/Collision/PhysicsFiltering.cpp
@@ -1,6 +1,7 @@
 // Copyright Epic Games, Inc. All Rights Reserved.
 
 #include "Physics/PhysicsFiltering.h"
+#include "CollisionQueryParams.h"
 
 
 //////////////////////////////////////////////////////////////////////////
diff --git a/Engine/Source/Runtime/Engine/Public/Physics/PhysicsFiltering.h b/Engine/Source/Runtime/Engine/Public/Physics/PhysicsFiltering.h
--- a/Engine/Source/Runtime/Engine/Public/Physics/PhysicsFiltering.h
+++ b/Engine/Source/Runtime/Engine/Public/Physics/PhysicsFiltering.h
@@ -7,6 +7,7 @@
 #include "PhysicsPublic.h"
 #include "PhysXIncludes.h"
 #include "Physics/PhysicsInterfaceTypes.h"
+#include "Chaos/ChaosEngineInterface.h"
 
 /** 
  * Set of flags stored in the PhysX FilterData
