Fix SUV suspension, wrong sign introduced on animation direction in CL 19320784.

diff --git a/Engine/Plugins/Experimental/ChaosVehiclesPlugin/Source/ChaosVehicles/Private/VehicleAnimationInstance.cpp b/Engine/Plugins/Experimental/ChaosVehiclesPlugin/Source/ChaosVehicles/Private/VehicleAnimationInstance.cpp
--- a/Engine/Plugins/Experimental/ChaosVehiclesPlugin/Source/ChaosVehicles/Private/VehicleAnimationInstance.cpp
+++ b/Engine/Plugins/Experimental/ChaosVehiclesPlugin/Source/ChaosVehicles/Private/VehicleAnimationInstance.cpp
@@ -128,7 +128,7 @@
 						WheelInstance.RotOffset.Yaw = VehicleWheel->GetSteerAngle();
 						WheelInstance.RotOffset.Roll = 0.f;
 
-						WheelInstance.LocOffset = VehicleWheel->GetSuspensionAxis() * VehicleWheel->GetSuspensionOffset();
+						WheelInstance.LocOffset = -VehicleWheel->GetSuspensionAxis() * VehicleWheel->GetSuspensionOffset();
 					}
 				}
 			}
