Fix issue with trace length when using sphere casts

diff --git a/Engine/Plugins/Experimental/ChaosVehiclesPlugin/Source/ChaosVehicles/Private/ChaosWheeledVehicleMovementComponent.cpp b/Engine/Plugins/Experimental/ChaosVehiclesPlugin/Source/ChaosVehicles/Private/ChaosWheeledVehicleMovementComponent.cpp
--- a/Engine/Plugins/Experimental/ChaosVehiclesPlugin/Source/ChaosVehicles/Private/ChaosWheeledVehicleMovementComponent.cpp
+++ b/Engine/Plugins/Experimental/ChaosVehiclesPlugin/Source/ChaosVehicles/Private/ChaosWheeledVehicleMovementComponent.cpp
@@ -405,11 +405,11 @@ void UChaosWheeledVehicleSimulation::PerformSuspensionTraces(const TArray<FSuspe
 					{
 					case ESweepShape::Spherecast:
 					{
-						float WheelRadius = PVehicle->Wheels[WheelIdx].GetEffectiveRadius(); // or wheel width
-						FVector VehicleUpAxis = TraceNormal;// GetOwner()->GetTransform().GetUnitAxis(EAxis::Z);
+						float WheelRadius = PVehicle->Wheels[WheelIdx].GetEffectiveRadius();
+						FVector VehicleUpAxis = TraceNormal;
 
 						FVector Start = TraceStart + VehicleUpAxis * WheelRadius;
-						FVector End = TraceEnd + VehicleUpAxis * WheelRadius;
+						FVector End = TraceEnd;
 
 						if (OverlapResult.Component->SweepComponent(ComponentHit, Start, End, FQuat::Identity, FCollisionShape::MakeSphere(WheelRadius), TraceParams.bTraceComplex))
 						{
@@ -463,15 +463,14 @@ void UChaosWheeledVehicleSimulation::PerformSuspensionTraces(const TArray<FSuspe
 			{
 			case ESweepShape::Spherecast:
 			{
-				//float Radius = PVehicle->Wheels[WheelIdx].GetEffectiveRadius(); // or wheel width
-				float Radius = PVehicle->Wheels[WheelIdx].Setup().WheelWidth * 0.5f; // or wheel width
-				FVector VehicleUpAxis = TraceNormal; //GetOwner()->GetTransform().GetUnitAxis(EAxis::Z);
+				float WheelRadius = PVehicle->Wheels[WheelIdx].GetEffectiveRadius();
+				FVector VehicleUpAxis = TraceNormal;
 
 				World->SweepSingleByChannel(HitResult
-					, TraceStart + VehicleUpAxis * Radius
-					, TraceEnd + VehicleUpAxis * Radius
+					, TraceStart + VehicleUpAxis * WheelRadius
+					, TraceEnd
 					, FQuat::Identity, SpringCollisionChannel
-					, FCollisionShape::MakeSphere(Radius), TraceParams
+					, FCollisionShape::MakeSphere(WheelRadius), TraceParams
 					, ResponseParams);
 			}
 			break;
