Exposing more vehicle interface to users, changing vehicle plugin so no longer beta status.

diff --git a/Engine/Plugins/Experimental/ChaosVehiclesPlugin/ChaosVehiclesPlugin.uplugin b/Engine/Plugins/Experimental/ChaosVehiclesPlugin/ChaosVehiclesPlugin.uplugin
--- a/Engine/Plugins/Experimental/ChaosVehiclesPlugin/ChaosVehiclesPlugin.uplugin
+++ b/Engine/Plugins/Experimental/ChaosVehiclesPlugin/ChaosVehiclesPlugin.uplugin
@@ -12,7 +12,7 @@
 	"SupportURL": "",
 	"EnabledByDefault": false,
 	"CanContainContent": true,
-	"IsBetaVersion": true,
+	"IsBetaVersion": false,
 	"IsExperimentalVersion": true,
 	"Installed": false,
 	"Modules": [
diff --git a/Engine/Plugins/Experimental/ChaosVehiclesPlugin/Source/ChaosVehicles/Public/ChaosVehicleManagerAsyncCallback.h b/Engine/Plugins/Experimental/ChaosVehiclesPlugin/Source/ChaosVehicles/Public/ChaosVehicleManagerAsyncCallback.h
--- a/Engine/Plugins/Experimental/ChaosVehiclesPlugin/Source/ChaosVehicles/Public/ChaosVehicleManagerAsyncCallback.h
+++ b/Engine/Plugins/Experimental/ChaosVehiclesPlugin/Source/ChaosVehicles/Public/ChaosVehicleManagerAsyncCallback.h
@@ -19,7 +19,7 @@ enum EChaosAsyncVehicleDataType : int8
 	AsyncDefault,
 };
 
-struct FWheelsOutput
+struct CHAOSVEHICLES_API FWheelsOutput
 {
 	FWheelsOutput()
 		: InContact(false)
@@ -68,7 +68,7 @@ struct FWheelsOutput
 /**
  * Per Vehicle Output State from Physics Thread to Game Thread                            
  */
-struct FPhysicsVehicleOutput
+struct CHAOSVEHICLES_API FPhysicsVehicleOutput
 {
 	FPhysicsVehicleOutput()
 		: CurrentGear(0)
@@ -92,7 +92,7 @@ struct FPhysicsVehicleOutput
 /**
  * Per Vehicle Input State from Game Thread to Physics Thread
  */                             
-struct FChaosVehicleAsyncInput
+struct CHAOSVEHICLES_API FChaosVehicleAsyncInput
 {
 	const EChaosAsyncVehicleDataType Type;
 	UChaosVehicleMovementComponent* Vehicle;
@@ -133,7 +133,7 @@ struct FChaosVehicleManagerAsyncInput : public Chaos::FSimCallbackInput
 /**
  * Async Output Data
  */
-struct FChaosVehicleAsyncOutput
+struct CHAOSVEHICLES_API FChaosVehicleAsyncOutput
 {
 	const EChaosAsyncVehicleDataType Type;
 	bool bValid;	// indicates no work was done
@@ -165,7 +165,7 @@ struct FChaosVehicleManagerAsyncOutput : public Chaos::FSimCallbackOutput
 /**
  * Async callback from the Physics Engine where we can perform our vehicle simulation
  */
-class FChaosVehicleManagerAsyncCallback : public Chaos::TSimCallbackObject<FChaosVehicleManagerAsyncInput, FChaosVehicleManagerAsyncOutput>
+class CHAOSVEHICLES_API FChaosVehicleManagerAsyncCallback : public Chaos::TSimCallbackObject<FChaosVehicleManagerAsyncInput, FChaosVehicleManagerAsyncOutput>
 {
 private:
 	virtual void OnPreSimulate_Internal() override;
diff --git a/Engine/Plugins/Experimental/ChaosVehiclesPlugin/Source/ChaosVehicles/Public/ChaosVehicleMovementComponent.h b/Engine/Plugins/Experimental/ChaosVehiclesPlugin/Source/ChaosVehicles/Public/ChaosVehicleMovementComponent.h
--- a/Engine/Plugins/Experimental/ChaosVehiclesPlugin/Source/ChaosVehicles/Public/ChaosVehicleMovementComponent.h
+++ b/Engine/Plugins/Experimental/ChaosVehiclesPlugin/Source/ChaosVehicles/Public/ChaosVehicleMovementComponent.h
@@ -30,7 +30,7 @@ struct FChaosVehicleAsyncInput;
 struct FChaosVehicleManagerAsyncOutput;
 
 
-struct FControlInputs
+struct CHAOSVEHICLES_API FControlInputs
 {
 	FControlInputs()
 		: SteeringInput(0.f)
@@ -79,7 +79,7 @@ struct FControlInputs
 };
 
 
-struct FVehicleDebugParams
+struct CHAOSVEHICLES_API FVehicleDebugParams
 {
 	bool ShowCOM = false;
 	bool ShowModelOrigin = false;
@@ -164,7 +164,7 @@ struct CHAOSVEHICLES_API FVehicleReplicatedState
 
 };
 
-struct FChaosVehicleDefaultAsyncInput : public FChaosVehicleAsyncInput
+struct CHAOSVEHICLES_API FChaosVehicleDefaultAsyncInput : public FChaosVehicleAsyncInput
 {
 	float GravityZ;
 	FControlInputs ControlInputs;
@@ -181,7 +181,7 @@ struct FChaosVehicleDefaultAsyncInput : public FChaosVehicleAsyncInput
 
 
 USTRUCT()
-struct FVehicleTorqueControlConfig
+struct CHAOSVEHICLES_API FVehicleTorqueControlConfig
 {
 public:
 	GENERATED_USTRUCT_BODY()
@@ -256,7 +256,7 @@ private:
 };
 
 USTRUCT()
-struct FVehicleTargetRotationControlConfig
+struct CHAOSVEHICLES_API FVehicleTargetRotationControlConfig
 {
 public:
 	GENERATED_USTRUCT_BODY()
@@ -353,7 +353,7 @@ private:
 };
 
 USTRUCT()
-struct FVehicleStabilizeControlConfig
+struct CHAOSVEHICLES_API FVehicleStabilizeControlConfig
 {
 public:
 	GENERATED_USTRUCT_BODY()
@@ -400,7 +400,7 @@ private:
 };
 
 /** Commonly used state - evaluated once used wherever required */
-struct FVehicleState
+struct CHAOSVEHICLES_API FVehicleState
 {
 	FVehicleState()
 		: VehicleWorldTransform(FTransform::Identity)
@@ -649,7 +649,7 @@ private:
 };
 
 USTRUCT()
-struct FVehicleThrustConfig
+struct CHAOSVEHICLES_API FVehicleThrustConfig
 {
 	GENERATED_USTRUCT_BODY()
 
@@ -716,7 +716,7 @@ private:
 
 };
 
-class UChaosVehicleSimulation
+class CHAOSVEHICLES_API UChaosVehicleSimulation
 {
 public:
 	virtual ~UChaosVehicleSimulation()
@@ -1341,7 +1341,6 @@ protected:
 
 	TUniquePtr<UChaosVehicleSimulation> VehicleSimulationPT;	/* simulation code running on the physics thread async callback */
 
-private:
 	UPROPERTY(transient, Replicated)
 	AController* OverrideController;
 
diff --git a/Engine/Plugins/Experimental/ChaosVehiclesPlugin/Source/ChaosVehicles/Public/ChaosWheeledVehicleMovementComponent.h b/Engine/Plugins/Experimental/ChaosVehiclesPlugin/Source/ChaosVehicles/Public/ChaosWheeledVehicleMovementComponent.h
--- a/Engine/Plugins/Experimental/ChaosVehiclesPlugin/Source/ChaosVehicles/Public/ChaosWheeledVehicleMovementComponent.h
+++ b/Engine/Plugins/Experimental/ChaosVehiclesPlugin/Source/ChaosVehicles/Public/ChaosWheeledVehicleMovementComponent.h
@@ -16,7 +16,7 @@
 PRAGMA_DISABLE_OPTIMIZATION
 #endif
 
-struct FWheeledVehicleDebugParams
+struct CHAOSVEHICLES_API FWheeledVehicleDebugParams
 {
 	bool ShowWheelCollisionNormal = false;
 	bool ShowSuspensionRaycasts = false;
@@ -49,7 +49,7 @@ struct FWheeledVehicleDebugParams
  * Advance through pages using p.Vehicles.NextDebugPage | p.Vehicles.PrevDebugPage which can be hooked
  * up to the keyboard or a controller in blueprint using execCommand
  */
-enum EDebugPages : uint8
+enum CHAOSVEHICLES_API EDebugPages : uint8
 {
 	BasicPage = 0,
 	PerformancePage,
@@ -161,7 +161,7 @@ struct CHAOSVEHICLES_API FWheelStatus
 };
 
 USTRUCT()
-struct FVehicleDifferentialConfig
+struct CHAOSVEHICLES_API FVehicleDifferentialConfig
 {
 	GENERATED_USTRUCT_BODY()
 
@@ -201,7 +201,7 @@ struct FVehicleDifferentialConfig
 };
 
 USTRUCT()
-struct FVehicleEngineConfig
+struct CHAOSVEHICLES_API FVehicleEngineConfig
 {
 	GENERATED_USTRUCT_BODY()
 
@@ -290,7 +290,7 @@ private:
 };
 
 USTRUCT()
-struct FVehicleTransmissionConfig
+struct CHAOSVEHICLES_API FVehicleTransmissionConfig
 {
 	GENERATED_USTRUCT_BODY()
 
@@ -420,7 +420,7 @@ enum class ESteeringType : uint8
 
 
 USTRUCT()
-struct FVehicleSteeringConfig
+struct CHAOSVEHICLES_API FVehicleSteeringConfig
 {
 	GENERATED_USTRUCT_BODY()
 
@@ -542,7 +542,7 @@ struct CHAOSVEHICLES_API FWheelState
 
 //////////////////////////////////////////////////////////////////////////
 
-class UChaosWheeledVehicleSimulation : public UChaosVehicleSimulation
+class CHAOSVEHICLES_API UChaosWheeledVehicleSimulation : public UChaosVehicleSimulation
 {
 public:
 
diff --git a/Engine/Plugins/Experimental/ChaosVehiclesPlugin/Source/ChaosVehicles/Public/DeferredForces.h b/Engine/Plugins/Experimental/ChaosVehiclesPlugin/Source/ChaosVehicles/Public/DeferredForces.h
--- a/Engine/Plugins/Experimental/ChaosVehiclesPlugin/Source/ChaosVehicles/Public/DeferredForces.h
+++ b/Engine/Plugins/Experimental/ChaosVehiclesPlugin/Source/ChaosVehicles/Public/DeferredForces.h
@@ -24,7 +24,7 @@ namespace Chaos
 }
 
 
-class FDeferredForces
+class CHAOSVEHICLES_API FDeferredForces
 {
 public:
 
diff --git a/Engine/Plugins/Experimental/ChaosVehiclesPlugin/Source/ChaosVehicles/Public/SnapshotData.h b/Engine/Plugins/Experimental/ChaosVehiclesPlugin/Source/ChaosVehicles/Public/SnapshotData.h
--- a/Engine/Plugins/Experimental/ChaosVehiclesPlugin/Source/ChaosVehicles/Public/SnapshotData.h
+++ b/Engine/Plugins/Experimental/ChaosVehiclesPlugin/Source/ChaosVehicles/Public/SnapshotData.h
@@ -7,7 +7,7 @@
 
 
 USTRUCT(BlueprintType)
-struct FBaseSnapshotData //: public UObject
+struct CHAOSVEHICLES_API FBaseSnapshotData //: public UObject
 {
 public:
 	GENERATED_BODY()
@@ -31,7 +31,7 @@ public:
 };
 
 USTRUCT(BlueprintType)
-struct FWheelSnapshot
+struct CHAOSVEHICLES_API FWheelSnapshot
 {
 	GENERATED_BODY()
 
@@ -61,7 +61,7 @@ struct FWheelSnapshot
 };
 
 USTRUCT(BlueprintType)
-struct FWheeledSnaphotData : public FBaseSnapshotData
+struct CHAOSVEHICLES_API FWheeledSnaphotData : public FBaseSnapshotData
 {
 public:
 	GENERATED_BODY()
