ChaosVehicle - Friction application point fix

diff --git a/Engine/Plugins/Experimental/ChaosVehiclesPlugin/Source/ChaosVehicles/Private/ChaosWheeledVehicleMovementComponent.cpp b/Engine/Plugins/Experimental/ChaosVehiclesPlugin/Source/ChaosVehicles/Private/ChaosWheeledVehicleMovementComponent.cpp
--- a/Engine/Plugins/Experimental/ChaosVehiclesPlugin/Source/ChaosVehicles/Private/ChaosWheeledVehicleMovementComponent.cpp
+++ b/Engine/Plugins/Experimental/ChaosVehiclesPlugin/Source/ChaosVehicles/Private/ChaosWheeledVehicleMovementComponent.cpp
@@ -18,6 +18,7 @@
 #include "Chaos/ChaosEngineInterface.h"
 #include "Chaos/PBDSuspensionConstraintData.h"
 #include "Chaos/DebugDrawQueue.h"
+#include "UObject/UE5MainStreamObjectVersion.h"
 
 #include "PhysicsProxy/SuspensionConstraintProxy.h"
 #include "PBDRigidsSolver.h"
@@ -519,14 +520,19 @@ void UChaosWheeledVehicleSimulation::ApplyWheelFrictionForces(float DeltaTime)
 			FVector GroundZVector = HitResult.Normal;
 			FVector GroundXVector = FVector::CrossProduct(VehicleState.VehicleRightAxis, GroundZVector);
 			FVector GroundYVector = FVector::CrossProduct(GroundZVector, GroundXVector);
-
-			// the force should be applied along the ground surface not along vehicle forward vector?
-			//FVector FrictionForceVector = VehicleState.VehicleWorldTransform.TransformVector(FrictionForceLocal);
-			FMatrix Mat(GroundXVector, GroundYVector, GroundZVector, VehicleState.VehicleWorldTransform.GetLocation());
+			
+			FMatrix Mat = FMatrix(GroundXVector, GroundYVector, GroundZVector, VehicleState.VehicleWorldTransform.GetLocation());
 			FVector FrictionForceVector = Mat.TransformVector(FrictionForceLocal);
 
 			check(PWheel.InContact());
-			AddForceAtPosition(FrictionForceVector, WheelState.WheelWorldLocation[WheelIdx]);
+			if (PVehicle->bLegacyWheelFrictionPosition)
+			{
+				AddForceAtPosition(FrictionForceVector, WheelState.WheelWorldLocation[WheelIdx]);
+			}
+			else
+			{
+				AddForceAtPosition(FrictionForceVector, HitResult.ImpactPoint);
+			}
 
 #if !(UE_BUILD_SHIPPING || UE_BUILD_TEST)
 			if (GWheeledVehicleDebugParams.ShowWheelForces)
@@ -1061,6 +1067,9 @@ UChaosWheeledVehicleMovementComponent::UChaosWheeledVehicleMovementComponent(con
 	bSuspensionEnabled = true;
 	bWheelFrictionEnabled = true;
 
+	// new vehicles don't use legacy method where friction forces are applied at wheel rather than wheel contact point 
+	bLegacyWheelFrictionPosition = false;
+
 	WheelTraceCollisionResponses = FCollisionResponseContainer::GetDefaultResponseContainer();
 	WheelTraceCollisionResponses.Vehicle = ECR_Ignore;
 }
@@ -1070,7 +1079,22 @@ void UChaosWheeledVehicleMovementComponent::Serialize(FArchive & Ar)
 {
 	Super::Serialize(Ar);
 
-	// custom serialization goes here..
+	Ar.UsingCustomVersion(FUE5MainStreamObjectVersion::GUID);
+}
+
+void UChaosWheeledVehicleMovementComponent::PostLoad()
+{
+	Super::PostLoad();
+
+#if WITH_EDITORONLY_DATA
+
+	if (GetLinkerCustomVersion(FUE5MainStreamObjectVersion::GUID) < FUE5MainStreamObjectVersion::VehicleFrictionForcePositionChange)
+	{
+		bLegacyWheelFrictionPosition = true;
+	}
+
+#endif  // #if WITH_EDITORONLY_DATA
+
 }
 
 #if WITH_EDITOR
@@ -1380,6 +1404,8 @@ void UChaosWheeledVehicleMovementComponent::SetupVehicle(TUniquePtr<Chaos::FSimp
 		PVehicle->Wheels[WheelIdx].SetWheelIndex(WheelIdx);
 		PVehicle->Suspension[WheelIdx].SetSpringIndex(WheelIdx);
 		PVehicle->NumDrivenWheels = NumDrivenWheels;
+
+		PVehicle->bLegacyWheelFrictionPosition = bLegacyWheelFrictionPosition;
 	}
 
 	RecalculateAxles();
diff --git a/Engine/Plugins/Experimental/ChaosVehiclesPlugin/Source/ChaosVehicles/Public/ChaosWheeledVehicleMovementComponent.h b/Engine/Plugins/Experimental/ChaosVehiclesPlugin/Source/ChaosVehicles/Public/ChaosWheeledVehicleMovementComponent.h
--- a/Engine/Plugins/Experimental/ChaosVehiclesPlugin/Source/ChaosVehicles/Public/ChaosWheeledVehicleMovementComponent.h
+++ b/Engine/Plugins/Experimental/ChaosVehiclesPlugin/Source/ChaosVehicles/Public/ChaosWheeledVehicleMovementComponent.h
@@ -40,8 +40,6 @@ struct CHAOSVEHICLES_API FWheeledVehicleDebugParams
 
 	float OverlapTestExpansionXY = 100.f;
 	float OverlapTestExpansionZ = 50.f;
-
-	//bool DisableSuspensionConstraint = false;
 };
 
 /**
@@ -627,6 +625,9 @@ class CHAOSVEHICLES_API UChaosWheeledVehicleMovementComponent : public UChaosVeh
 	UPROPERTY(EditAnywhere, Category = WheelSetup)
 	bool bWheelFrictionEnabled;
 
+	UPROPERTY(EditAnywhere, Category = WheelSetup)
+	bool bLegacyWheelFrictionPosition;
+
 	/** Wheels to create */
 	UPROPERTY(EditAnywhere, Category = WheelSetup)
 	TArray<FChaosWheelSetup> WheelSetups;
@@ -722,7 +723,8 @@ class CHAOSVEHICLES_API UChaosWheeledVehicleMovementComponent : public UChaosVeh
 	//////////////////////////////////////////////////////////////////////////
 	// Public
 
-	virtual void Serialize(FArchive & Ar) override;
+	virtual void Serialize(FArchive& Ar) override;
+	virtual void PostLoad() override;
 
 	// Get output data from Physics Thread
 	virtual void ParallelUpdate(float DeltaSeconds);
diff --git a/Engine/Source/Runtime/Core/Public/UObject/UE5MainStreamObjectVersion.h b/Engine/Source/Runtime/Core/Public/UObject/UE5MainStreamObjectVersion.h
--- a/Engine/Source/Runtime/Core/Public/UObject/UE5MainStreamObjectVersion.h
+++ b/Engine/Source/Runtime/Core/Public/UObject/UE5MainStreamObjectVersion.h
@@ -271,6 +271,9 @@ struct CORE_API FUE5MainStreamObjectVersion
 		// Add damage propagation settings to geometry collections
 		GeometryCollectionDamagePropagationData,
 
+		// Wheel friction forces are now applied at tire contact point
+		VehicleFrictionForcePositionChange,
+
 		// -----<new versions can be added above this line>-------------------------------------------------
 		VersionPlusOne,
 		LatestVersion = VersionPlusOne - 1
diff --git a/Engine/Source/Runtime/Experimental/ChaosVehicles/ChaosVehiclesCore/Public/SimpleVehicle.h b/Engine/Source/Runtime/Experimental/ChaosVehicles/ChaosVehiclesCore/Public/SimpleVehicle.h
--- a/Engine/Source/Runtime/Experimental/ChaosVehicles/ChaosVehiclesCore/Public/SimpleVehicle.h
+++ b/Engine/Source/Runtime/Experimental/ChaosVehicles/ChaosVehiclesCore/Public/SimpleVehicle.h
@@ -40,6 +40,7 @@ public:
 	, bMechanicalSimEnabled(true)
 	, bWheelFrictionEnabled(true)
 	, NumDrivenWheels(0)
+	, bLegacyWheelFrictionPosition(false)
 	{
 
 	}
@@ -185,6 +186,7 @@ public:
 	bool bMechanicalSimEnabled;
 	bool bWheelFrictionEnabled;
 	uint32 NumDrivenWheels;
+	bool bLegacyWheelFrictionPosition;
 };
 
 
