ChaosVehiclePlugin: Fix networking issue with gears

diff --git a/Engine/Plugins/Experimental/ChaosVehiclesPlugin/Source/ChaosVehicles/Private/ChaosVehicleMovementComponent.cpp b/Engine/Plugins/Experimental/ChaosVehiclesPlugin/Source/ChaosVehicles/Private/ChaosVehicleMovementComponent.cpp
--- a/Engine/Plugins/Experimental/ChaosVehiclesPlugin/Source/ChaosVehicles/Private/ChaosVehicleMovementComponent.cpp
+++ b/Engine/Plugins/Experimental/ChaosVehiclesPlugin/Source/ChaosVehicles/Private/ChaosVehicleMovementComponent.cpp
@@ -1155,7 +1155,7 @@ void UChaosVehicleMovementComponent::UpdateState(float DeltaTime)
 		HandbrakeInput = HandbrakeInputRate.InterpInputValue(DeltaTime, HandbrakeInput, CalcHandbrakeInput());
 
 		// and send to server - (ServerUpdateState_Implementation below)
-		ServerUpdateState(SteeringInput, ThrottleInput, BrakeInput, HandbrakeInput, GetTargetGear(), RollInput, PitchInput, YawInput);
+		ServerUpdateState(SteeringInput, ThrottleInput, BrakeInput, HandbrakeInput, GetCurrentGear(), RollInput, PitchInput, YawInput);
 
 		if (PawnOwner && PawnOwner->IsNetMode(NM_Client))
 		{
