Added interface to override the drive and braking torque per wheel.

diff --git a/Engine/Plugins/Experimental/ChaosVehiclesPlugin/Source/ChaosVehicles/Private/ChaosVehicleWheel.cpp b/Engine/Plugins/Experimental/ChaosVehiclesPlugin/Source/ChaosVehicles/Private/ChaosVehicleWheel.cpp
--- a/Engine/Plugins/Experimental/ChaosVehiclesPlugin/Source/ChaosVehicles/Private/ChaosVehicleWheel.cpp
+++ b/Engine/Plugins/Experimental/ChaosVehiclesPlugin/Source/ChaosVehicles/Private/ChaosVehicleWheel.cpp
@@ -48,6 +48,7 @@ UChaosVehicleWheel::UChaosVehicleWheel(const FObjectInitializer& ObjectInitializ
 	WheelLoadRatio = 0.5f;
 	RollbarScaling = 0.15f;
 	SweepType = ESweepType::SimpleSweep;
+	ExternalTorqueCombineMethod = ETorqueCombineMethod::None;
 }
 
 
diff --git a/Engine/Plugins/Experimental/ChaosVehiclesPlugin/Source/ChaosVehicles/Private/ChaosWheeledVehicleMovementComponent.cpp b/Engine/Plugins/Experimental/ChaosVehiclesPlugin/Source/ChaosVehicles/Private/ChaosWheeledVehicleMovementComponent.cpp
--- a/Engine/Plugins/Experimental/ChaosVehiclesPlugin/Source/ChaosVehicles/Private/ChaosWheeledVehicleMovementComponent.cpp
+++ b/Engine/Plugins/Experimental/ChaosVehiclesPlugin/Source/ChaosVehicles/Private/ChaosWheeledVehicleMovementComponent.cpp
@@ -2599,6 +2599,56 @@ void UChaosWheeledVehicleMovementComponent::SetWheelMaxSteerAngle(int WheelIndex
 	}
 }
 
+void UChaosWheeledVehicleMovementComponent::SetTorqueCombineMethod(ETorqueCombineMethod InCombineMethod, int32 WheelIndex)
+{
+	if (FBodyInstance* TargetInstance = GetBodyInstance())
+	{
+		FPhysicsCommand::ExecuteWrite(TargetInstance->ActorHandle, [&](const FPhysicsActorHandle& Chassis)
+			{
+				if (VehicleSimulationPT && VehicleSimulationPT->PVehicle && WheelIndex < VehicleSimulationPT->PVehicle->Wheels.Num())
+				{
+					Chaos::FSimpleWheelSim& VehicleWheel = VehicleSimulationPT->PVehicle->Wheels[WheelIndex];
+
+					VehicleWheel.SetTorqueCombineMethod(static_cast<Chaos::FSimpleWheelConfig::EExternalTorqueCombineMethod>(InCombineMethod));
+				}
+			});
+	}
+}
+
+void UChaosWheeledVehicleMovementComponent::SetDriveTorque(float DriveTorque, int32 WheelIndex)
+{
+	SetSleeping(false);
+
+	if (FBodyInstance* TargetInstance = GetBodyInstance())
+	{
+		FPhysicsCommand::ExecuteWrite(TargetInstance->ActorHandle, [&](const FPhysicsActorHandle& Chassis)
+			{
+				if (VehicleSimulationPT && VehicleSimulationPT->PVehicle && WheelIndex < VehicleSimulationPT->PVehicle->Wheels.Num())
+				{
+					Chaos::FSimpleWheelSim& VehicleWheel = VehicleSimulationPT->PVehicle->Wheels[WheelIndex];
+
+					VehicleWheel.SetDriveTorqueOverride(TorqueMToCm(DriveTorque));
+				}
+			});
+	}
+}
+
+void UChaosWheeledVehicleMovementComponent::SetBrakeTorque(float BrakeTorque, int32 WheelIndex)
+{
+	if (FBodyInstance* TargetInstance = GetBodyInstance())
+	{
+		FPhysicsCommand::ExecuteWrite(TargetInstance->ActorHandle, [&](const FPhysicsActorHandle& Chassis)
+			{
+				if (VehicleSimulationPT && VehicleSimulationPT->PVehicle && WheelIndex < VehicleSimulationPT->PVehicle->Wheels.Num())
+				{
+					Chaos::FSimpleWheelSim& VehicleWheel = VehicleSimulationPT->PVehicle->Wheels[WheelIndex];
+
+					VehicleWheel.SetBrakeTorqueOverride(TorqueMToCm(BrakeTorque));
+				}
+			});
+	}
+}
+
 
 float UChaosWheeledVehicleMovementComponent::GetSuspensionOffset(int WheelIndex)
 {
diff --git a/Engine/Plugins/Experimental/ChaosVehiclesPlugin/Source/ChaosVehicles/Public/ChaosVehicleWheel.h b/Engine/Plugins/Experimental/ChaosVehiclesPlugin/Source/ChaosVehicles/Public/ChaosVehicleWheel.h
--- a/Engine/Plugins/Experimental/ChaosVehiclesPlugin/Source/ChaosVehicles/Public/ChaosVehicleWheel.h
+++ b/Engine/Plugins/Experimental/ChaosVehiclesPlugin/Source/ChaosVehicles/Public/ChaosVehicleWheel.h
@@ -55,6 +55,17 @@ class UChaosWheeledVehicleMovementComponent;
 		Rear
 	};
 
+	UENUM()
+	enum class ETorqueCombineMethod : uint8
+	{
+		/** External torque value has no effect - default **/
+		None = 0,
+		/** completely replace existing torques, can set and forget will apply same torque every frame until zeroed */
+		Override,
+		/** Combine external torque with existing torques, must set external torque every frame */
+		Additive
+	};
+
 	UCLASS(BlueprintType, Blueprintable)
 		class CHAOSVEHICLES_API UChaosVehicleWheel : public UObject
 	{
@@ -134,6 +145,10 @@ class UChaosWheeledVehicleMovementComponent;
 		UPROPERTY(EditAnywhere, Category = Wheel)
 		bool bTractionControlEnabled;
 
+		/** Determines how the SetDriveTorque/SetBrakeTorque inputs are combined with the internal torques */
+		UPROPERTY(EditAnywhere, Category = Wheel)
+		ETorqueCombineMethod ExternalTorqueCombineMethod;
+
 		UPROPERTY(EditAnywhere, Category = Setup)
 		FRuntimeFloatCurve LateralSlipGraph;
 
@@ -339,6 +354,7 @@ class UChaosWheeledVehicleMovementComponent;
 			PWheelConfig.SideSlipModifier = this->SideSlipModifier;
 			PWheelConfig.SlipThreshold = this->SlipThreshold;
 			PWheelConfig.SkidThreshold = this->SkidThreshold;
+			PWheelConfig.ExternalTorqueCombineMethod = static_cast<Chaos::FSimpleWheelConfig::EExternalTorqueCombineMethod>(this->ExternalTorqueCombineMethod);
 
 			PWheelConfig.LateralSlipGraph.Empty();
 			float NumSamples = 20;
diff --git a/Engine/Plugins/Experimental/ChaosVehiclesPlugin/Source/ChaosVehicles/Public/ChaosWheeledVehicleMovementComponent.h b/Engine/Plugins/Experimental/ChaosVehiclesPlugin/Source/ChaosVehicles/Public/ChaosWheeledVehicleMovementComponent.h
--- a/Engine/Plugins/Experimental/ChaosVehiclesPlugin/Source/ChaosVehicles/Public/ChaosWheeledVehicleMovementComponent.h
+++ b/Engine/Plugins/Experimental/ChaosVehiclesPlugin/Source/ChaosVehicles/Public/ChaosWheeledVehicleMovementComponent.h
@@ -830,6 +830,15 @@ class CHAOSVEHICLES_API UChaosWheeledVehicleMovementComponent : public UChaosVeh
 	UFUNCTION(BlueprintCallable, Category = "Game|Components|ChaosWheeledVehicleMovement")
 	void SetWheelMaxSteerAngle(int WheelIndex, float AngleDegrees);
 
+	UFUNCTION(BlueprintCallable, Category = "Game|Components|ChaosWheeledVehicleMovement")
+	void SetTorqueCombineMethod(ETorqueCombineMethod InCombineMethod, int32 WheelIndex);
+
+	UFUNCTION(BlueprintCallable, Category = "Game|Components|ChaosWheeledVehicleMovement")
+	void SetDriveTorque(float DriveTorque, int32 WheelIndex);
+	
+	UFUNCTION(BlueprintCallable, Category = "Game|Components|ChaosWheeledVehicleMovement")
+	void SetBrakeTorque(float BrakeTorque, int32 WheelIndex);
+
 	/** */
 	virtual TUniquePtr<Chaos::FSimpleWheeledVehicle> CreatePhysicsVehicle() override
 	{
diff --git a/Engine/Source/Runtime/Experimental/ChaosVehicles/ChaosVehiclesCore/Private/WheelSystem.cpp b/Engine/Source/Runtime/Experimental/ChaosVehicles/ChaosVehiclesCore/Private/WheelSystem.cpp
--- a/Engine/Source/Runtime/Experimental/ChaosVehicles/ChaosVehiclesCore/Private/WheelSystem.cpp
+++ b/Engine/Source/Runtime/Experimental/ChaosVehicles/ChaosVehiclesCore/Private/WheelSystem.cpp
@@ -21,6 +21,7 @@ namespace Chaos
 		, MaxSteeringAngle(SetupIn->MaxSteeringAngle)
 		, MaxBrakeTorque(SetupIn->MaxBrakeTorque)
 		, HandbrakeTorque(SetupIn->HandbrakeTorque)
+		, ExternalTorqueCombineMethod(SetupIn->ExternalTorqueCombineMethod)
 		, Re(SetupIn->WheelRadius)
 		, Omega(0.f)
 		, Sx(0.f)
@@ -32,6 +33,8 @@ namespace Chaos
 		, AngularPosition(0.f)
 		, SteeringAngle(0.f)
 		, SurfaceFriction(1.f)
+		, ExternalDriveTorque(0.0f)
+		, ExternalBrakeTorque(0.0f)
 		, ForceFromFriction(FVector::ZeroVector)
 		, MassPerWheel(250.f)
 		, SlipVelocity(0.f)
@@ -55,6 +58,21 @@ namespace Chaos
 		// X is longitudinal direction, Y is lateral
 		SlipAngle = FVehicleUtility::CalculateSlipAngle(GroundVelocityVector.Y, GroundVelocityVector.X);
 
+		if (ExternalTorqueCombineMethod == FSimpleWheelConfig::EExternalTorqueCombineMethod::Override)
+		{ 
+			// Set and forget or per frame torque override
+			DriveTorque = ExternalDriveTorque;
+			BrakeTorque = ExternalBrakeTorque;
+		}
+		else if (ExternalTorqueCombineMethod == FSimpleWheelConfig::EExternalTorqueCombineMethod::Additive)
+		{
+			// per frame torque added to engine torque
+			DriveTorque += ExternalDriveTorque;
+			BrakeTorque += ExternalBrakeTorque;
+			ExternalDriveTorque = 0.0f;
+			ExternalBrakeTorque = 0.0f;
+		}
+
 		// The physics system is mostly unit-less i.e. can work in meters or cm, however there are 
 		// a couple of places where the results are wrong if Cm is used. This is one of them, the simulated radius
 		// for torque must be real size to obtain the correct output values.
diff --git a/Engine/Source/Runtime/Experimental/ChaosVehicles/ChaosVehiclesCore/Public/WheelSystem.h b/Engine/Source/Runtime/Experimental/ChaosVehicles/ChaosVehiclesCore/Public/WheelSystem.h
--- a/Engine/Source/Runtime/Experimental/ChaosVehicles/ChaosVehiclesCore/Public/WheelSystem.h
+++ b/Engine/Source/Runtime/Experimental/ChaosVehicles/ChaosVehiclesCore/Public/WheelSystem.h
@@ -72,6 +72,13 @@ struct CHAOSVEHICLESCORE_API FSimpleWheelConfig
 		Rear
 	};
 	
+	enum EExternalTorqueCombineMethod : uint8
+	{
+		None = 0,
+		Override,
+		Additive
+	};
+
 	FSimpleWheelConfig() 
 		: Offset(FVector(2.f, 1.f, 0.f))
 		, WheelMass(20.f) // [Kg]
@@ -89,6 +96,7 @@ struct CHAOSVEHICLESCORE_API FSimpleWheelConfig
 		, TorqueRatio(0.f)
 		, AxleType(EAxleType::UndefinedAxle)
 		, FrictionCombineMethod(EFrictionCombineMethod::Multiply)
+		, ExternalTorqueCombineMethod(EExternalTorqueCombineMethod::None)
 		, FrictionMultiplier(2.0f)
 		, LateralSlipGraphMultiplier(1.0f)
 		, CorneringStiffness(1000.0f)
@@ -127,6 +135,7 @@ struct CHAOSVEHICLESCORE_API FSimpleWheelConfig
 	EAxleType AxleType;
 
 	EFrictionCombineMethod FrictionCombineMethod; //#todo: use this variable
+	EExternalTorqueCombineMethod ExternalTorqueCombineMethod;
 
 	float FrictionMultiplier;
 	float LateralSlipGraphMultiplier;
@@ -193,6 +202,22 @@ public:
 		DriveTorque = EngineTorqueIn;
 	}
 
+	void SetTorqueCombineMethod(FSimpleWheelConfig::EExternalTorqueCombineMethod InCombineMethod)
+	{
+		ExternalTorqueCombineMethod = InCombineMethod;
+	}
+
+	void SetBrakeTorqueOverride(float BrakeTorqueIn, bool bEngineBrakingIn = false)
+	{
+		ExternalBrakeTorque = BrakeTorqueIn;
+		bEngineBraking = bEngineBrakingIn;
+	}
+
+	void SetDriveTorqueOverride(float EngineTorqueIn)
+	{
+		ExternalDriveTorque = EngineTorqueIn;
+	}
+
 	/** Set the vehicle's speed at the wheels location in local wheel coords */
 	void SetVehicleGroundSpeed(const FVector& VIn)
 	{
@@ -401,6 +426,7 @@ public:
 	float MaxSteeringAngle;
 	float MaxBrakeTorque;
 	float HandbrakeTorque;
+	FSimpleWheelConfig::EExternalTorqueCombineMethod ExternalTorqueCombineMethod;
 
 	float Re;		// [cm] Effective Wheel Radius could change dynamically if get a flat?, tire shreds
 	float Omega;	// [radians/sec] Wheel Rotation Angular Velocity
@@ -416,6 +442,9 @@ public:
 	float SurfaceFriction;
 	float MaxOmega;
 
+	float ExternalDriveTorque;		// [N.m]
+	float ExternalBrakeTorque;		// [N.m]
+
 	FVector ForceFromFriction;
 	float MassPerWheel;
 
