Manual gears were not changing fully and flicking, after networking fix CL 21887301

diff --git a/Engine/Plugins/Experimental/ChaosVehiclesPlugin/Source/ChaosVehicles/Private/ChaosVehicleMovementComponent.cpp b/Engine/Plugins/Experimental/ChaosVehiclesPlugin/Source/ChaosVehicles/Private/ChaosVehicleMovementComponent.cpp
--- a/Engine/Plugins/Experimental/ChaosVehiclesPlugin/Source/ChaosVehicles/Private/ChaosVehicleMovementComponent.cpp
+++ b/Engine/Plugins/Experimental/ChaosVehiclesPlugin/Source/ChaosVehicles/Private/ChaosVehicleMovementComponent.cpp
@@ -1254,7 +1254,13 @@ void UChaosVehicleMovementComponent::ServerUpdateState_Implementation(float InSt
 
 	if (!GetUseAutoGears())
 	{
-		SetTargetGear(InCurrentGear, true);
+		if (UWorld* World = GetWorld())
+		{
+			if (World->GetNetMode() == NM_DedicatedServer)
+			{
+				SetTargetGear(InCurrentGear, true);
+			}
+		}
 	}
 
 	// update state of inputs
