Fix potential crash when dynamically deleting vehicles.

diff --git a/Engine/Plugins/Experimental/ChaosVehiclesPlugin/Source/ChaosVehicles/Private/ChaosVehicleMovementComponent.cpp b/Engine/Plugins/Experimental/ChaosVehiclesPlugin/Source/ChaosVehicles/Private/ChaosVehicleMovementComponent.cpp
--- a/Engine/Plugins/Experimental/ChaosVehiclesPlugin/Source/ChaosVehicles/Private/ChaosVehicleMovementComponent.cpp
+++ b/Engine/Plugins/Experimental/ChaosVehiclesPlugin/Source/ChaosVehicles/Private/ChaosVehicleMovementComponent.cpp
@@ -1812,6 +1812,10 @@ void UChaosVehicleMovementComponent::ParallelUpdate(float DeltaSeconds)
 					PVehicleOutput->Wheels[WheelIdx].SuspensionOffset = FMath::Lerp(Current.SuspensionOffset, Next.SuspensionOffset, OutputInterpAlpha);
 					PVehicleOutput->Wheels[WheelIdx].SpringForce = FMath::Lerp(Current.SpringForce, Next.SpringForce, OutputInterpAlpha);
 					PVehicleOutput->Wheels[WheelIdx].NormalizedSuspensionLength = FMath::Lerp(Current.NormalizedSuspensionLength, Next.NormalizedSuspensionLength, OutputInterpAlpha);
+
+					PVehicleOutput->Wheels[WheelIdx].bBlockingHit = Current.bBlockingHit;
+					PVehicleOutput->Wheels[WheelIdx].ImpactPoint = FMath::Lerp(Current.ImpactPoint, Next.ImpactPoint, OutputInterpAlpha);
+					PVehicleOutput->Wheels[WheelIdx].PhysMaterial = Current.PhysMaterial;
 				}
 			}
 			else // WHEN ASYNC IS OFF IT STILL GENERATES THE ASYNC CALLBACK BUT THERE IS ONLY EVER THE CURRENT AND NO NEXT OUTPUT TO INTERPOLATE BETWEEN
@@ -1843,6 +1847,10 @@ void UChaosVehicleMovementComponent::ParallelUpdate(float DeltaSeconds)
 					PVehicleOutput->Wheels[WheelIdx].SuspensionOffset = Current.SuspensionOffset;
 					PVehicleOutput->Wheels[WheelIdx].SpringForce = Current.SpringForce;
 					PVehicleOutput->Wheels[WheelIdx].NormalizedSuspensionLength = Current.NormalizedSuspensionLength;
+
+					PVehicleOutput->Wheels[WheelIdx].bBlockingHit = Current.bBlockingHit;
+					PVehicleOutput->Wheels[WheelIdx].ImpactPoint = Current.ImpactPoint;
+					PVehicleOutput->Wheels[WheelIdx].PhysMaterial = Current.PhysMaterial;
 				}
 
 			}
diff --git a/Engine/Plugins/Experimental/ChaosVehiclesPlugin/Source/ChaosVehicles/Private/ChaosVehicleWheel.cpp b/Engine/Plugins/Experimental/ChaosVehiclesPlugin/Source/ChaosVehicles/Private/ChaosVehicleWheel.cpp
--- a/Engine/Plugins/Experimental/ChaosVehiclesPlugin/Source/ChaosVehicles/Private/ChaosVehicleWheel.cpp
+++ b/Engine/Plugins/Experimental/ChaosVehiclesPlugin/Source/ChaosVehicles/Private/ChaosVehicleWheel.cpp
@@ -162,14 +162,8 @@ void UChaosVehicleWheel::PostEditChangeProperty( FPropertyChangedEvent& Property
 
 UPhysicalMaterial* UChaosVehicleWheel::GetContactSurfaceMaterial()
 {
-	UPhysicalMaterial* PhysMaterial = NULL;
-
-	if (HitResult.bBlockingHit && HitResult.PhysMaterial.IsValid())
-	{
-		PhysMaterial = HitResult.PhysMaterial.Get();
-	}
-
-	return PhysMaterial;
+	check(VehicleComponent && VehicleComponent->PhysicsVehicleOutput());
+	return VehicleComponent->GetPhysMaterial(WheelIndex);
 }
 
 
diff --git a/Engine/Plugins/Experimental/ChaosVehiclesPlugin/Source/ChaosVehicles/Private/ChaosWheeledVehicleMovementComponent.cpp b/Engine/Plugins/Experimental/ChaosVehiclesPlugin/Source/ChaosVehicles/Private/ChaosWheeledVehicleMovementComponent.cpp
--- a/Engine/Plugins/Experimental/ChaosVehiclesPlugin/Source/ChaosVehicles/Private/ChaosWheeledVehicleMovementComponent.cpp
+++ b/Engine/Plugins/Experimental/ChaosVehiclesPlugin/Source/ChaosVehicles/Private/ChaosWheeledVehicleMovementComponent.cpp
@@ -164,7 +164,7 @@ bool UChaosWheeledVehicleSimulation::CanSimulate() const
 
 	return (PVehicle && PVehicle.IsValid()
 		&& PVehicle->Engine.Num() == PVehicle->Transmission.Num()
-		&& Wheels.Num() > 0 && Wheels.Num() == PVehicle->Suspension.Num());
+		&& PVehicle->Wheels.Num() == PVehicle->Suspension.Num());
 }
 
 void UChaosWheeledVehicleSimulation::TickVehicle(UWorld* WorldIn, float DeltaTime, const FChaosVehicleDefaultAsyncInput& InputData, FChaosVehicleAsyncOutput& OutputData, Chaos::FRigidBodyHandle_Internal* Handle)
@@ -184,16 +184,14 @@ void UChaosWheeledVehicleSimulation::UpdateSimulation(float DeltaTime, const FCh
 	if (CanSimulate() && Handle)
 	{
 		// sanity check that everything is setup ok
-		ensure(Wheels.Num() == PVehicle->Suspension.Num());
-		ensure(Wheels.Num() == PVehicle->Wheels.Num());
-		ensure(WheelState.LocalWheelVelocity.Num() == Wheels.Num());
-		ensure(WheelState.WheelWorldLocation.Num() == Wheels.Num());
-		ensure(WheelState.WorldWheelVelocity.Num() == Wheels.Num());
+		ensure(PVehicle->Wheels.Num() == PVehicle->Suspension.Num());
+		ensure(WheelState.LocalWheelVelocity.Num() == PVehicle->Wheels.Num());
+		ensure(WheelState.WheelWorldLocation.Num() == PVehicle->Wheels.Num());
+		ensure(WheelState.WorldWheelVelocity.Num() == PVehicle->Wheels.Num());
 
 		///////////////////////////////////////////////////////////////////////
 		// Cache useful state so we are not re-calculating the same data
-
-		for (int WheelIdx = 0; WheelIdx < Wheels.Num(); WheelIdx++)
+		for (int WheelIdx = 0; WheelIdx < PVehicle->Suspension.Num(); WheelIdx++)
 		{
 			bool bCaptured = false;
 
@@ -225,7 +223,7 @@ void UChaosWheeledVehicleSimulation::UpdateSimulation(float DeltaTime, const FCh
 		///////////////////////////////////////////////////////////////////////
 		// Suspension Raycast
 
-		for (int WheelIdx = 0; WheelIdx < Wheels.Num(); WheelIdx++)
+		for (int WheelIdx = 0; WheelIdx < PVehicle->Suspension.Num(); WheelIdx++)
 		{
 			auto& PSuspension = PVehicle->Suspension[WheelIdx];
 			auto& PWheel = PVehicle->Wheels[WheelIdx];
@@ -234,7 +232,7 @@ void UChaosWheeledVehicleSimulation::UpdateSimulation(float DeltaTime, const FCh
 
 		if (!GWheeledVehicleDebugParams.DisableSuspensionForces && PVehicle->bSuspensionEnabled)
 		{
-			PerformSuspensionTraces(WheelState.Trace, InputData.TraceParams, InputData.TraceCollisionResponse);
+			PerformSuspensionTraces(WheelState.Trace, InputData.TraceParams, InputData.TraceCollisionResponse, InputData.WheelTraceParams);
 		}
 
 		//////////////////////////////////////////////////////////////////////////
@@ -242,10 +240,10 @@ void UChaosWheeledVehicleSimulation::UpdateSimulation(float DeltaTime, const FCh
 
 		VehicleState.bVehicleInAir = true;
 		VehicleState.NumWheelsOnGround = 0;
-		for (int WheelIdx = 0; WheelIdx < Wheels.Num(); WheelIdx++)
+		for (int WheelIdx = 0; WheelIdx < PVehicle->Wheels.Num(); WheelIdx++)
 		{
 			// tell systems who care that wheel is touching the ground
-			PVehicle->Wheels[WheelIdx].SetOnGround(Wheels[WheelIdx]->HitResult.bBlockingHit);
+			PVehicle->Wheels[WheelIdx].SetOnGround(WheelState.TraceResult[WheelIdx].bBlockingHit);
 
 			// only requires one wheel to be on the ground for the vehicle to be NOT in the air
 			if (PVehicle->Wheels[WheelIdx].InContact())
@@ -254,7 +252,7 @@ void UChaosWheeledVehicleSimulation::UpdateSimulation(float DeltaTime, const FCh
 				VehicleState.NumWheelsOnGround++;
 			}
 		}
-		VehicleState.bAllWheelsOnGround = (VehicleState.NumWheelsOnGround == Wheels.Num());
+		VehicleState.bAllWheelsOnGround = (VehicleState.NumWheelsOnGround == PVehicle->Wheels.Num());
 
 		///////////////////////////////////////////////////////////////////////
 		// Input
@@ -302,7 +300,7 @@ bool UChaosWheeledVehicleSimulation::ContainsTraces(const FBox& Box, const TArra
 {
 	const FAABB3 Aabb(Box.Min, Box.Max);
 
-	for (int WheelIdx = 0; WheelIdx < Wheels.Num(); WheelIdx++)
+	for (int WheelIdx = 0; WheelIdx < SuspensionTrace.Num(); WheelIdx++)
 	{
 		if (!Aabb.Contains(SuspensionTrace[WheelIdx].Start) || !Aabb.Contains(SuspensionTrace[WheelIdx].End))
 		{
@@ -313,7 +311,7 @@ bool UChaosWheeledVehicleSimulation::ContainsTraces(const FBox& Box, const TArra
 	return true;
 }
 
-void UChaosWheeledVehicleSimulation::PerformSuspensionTraces(const TArray<FSuspensionTrace>& SuspensionTrace, FCollisionQueryParams& TraceParams, FCollisionResponseContainer& CollisionResponse)
+void UChaosWheeledVehicleSimulation::PerformSuspensionTraces(const TArray<FSuspensionTrace>& SuspensionTrace, FCollisionQueryParams& TraceParams, FCollisionResponseContainer& CollisionResponse, TArray<FWheelTraceParams>& WheelTraceParams)
 {
 	SCOPE_CYCLE_COUNTER(STAT_ChaosVehicle_SuspensionRaycasts);
 
@@ -333,7 +331,7 @@ void UChaosWheeledVehicleSimulation::PerformSuspensionTraces(const TArray<FSuspe
 			QueryBox.Init();
 
 			//FBox QueryBox;
-			for (int WheelIdx = 0; WheelIdx < Wheels.Num(); WheelIdx++)
+			for (int WheelIdx = 0; WheelIdx < SuspensionTrace.Num(); WheelIdx++)
 			{
 				const FVector& TraceStart = SuspensionTrace[WheelIdx].Start;
 				const FVector& TraceEnd = SuspensionTrace[WheelIdx].End;
@@ -375,16 +373,16 @@ void UChaosWheeledVehicleSimulation::PerformSuspensionTraces(const TArray<FSuspe
 	#endif
 
 		SCOPE_CYCLE_COUNTER(STAT_ChaosVehicle_SuspensionTraces);
-		for (int32 WheelIdx = 0; WheelIdx < Wheels.Num(); ++WheelIdx)
+		for (int32 WheelIdx = 0; WheelIdx < SuspensionTrace.Num(); ++WheelIdx)
 		{
-			FHitResult& HitResult = Wheels[WheelIdx]->HitResult;
+			FHitResult& HitResult = WheelState.TraceResult[WheelIdx];
 			HitResult = FHitResult();
 
 			if (bOverlapHit)
 			{
 				const FVector& TraceStart = SuspensionTrace[WheelIdx].Start;
 				const FVector& TraceEnd = SuspensionTrace[WheelIdx].End;
-				TraceParams.bTraceComplex = (Wheels[WheelIdx]->SweepType == ESweepType::ComplexSweep);
+				TraceParams.bTraceComplex = (WheelTraceParams[WheelIdx].SweepType == ESweepType::ComplexSweep);
 
 				if (GWheeledVehicleDebugParams.TraceTypeOverride > 0)
 				{
@@ -402,7 +400,7 @@ void UChaosWheeledVehicleSimulation::PerformSuspensionTraces(const TArray<FSuspe
 
 					FHitResult ComponentHit;
 
-					switch (Wheels[WheelIdx]->SweepShape)
+					switch (WheelTraceParams[WheelIdx].SweepShape)
 					{
 					case ESweepShape::Spherecast:
 					{
@@ -444,13 +442,13 @@ void UChaosWheeledVehicleSimulation::PerformSuspensionTraces(const TArray<FSuspe
 	else
 	{
 		SCOPE_CYCLE_COUNTER(STAT_ChaosVehicle_SuspensionTraces);
-		for (int WheelIdx = 0; WheelIdx < Wheels.Num(); WheelIdx++)
+		for (int WheelIdx = 0; WheelIdx < SuspensionTrace.Num(); WheelIdx++)
 		{
-			FHitResult& HitResult = Wheels[WheelIdx]->HitResult;
+			FHitResult& HitResult = WheelState.TraceResult[WheelIdx];
 
 			FVector TraceStart = SuspensionTrace[WheelIdx].Start;
 			FVector TraceEnd = SuspensionTrace[WheelIdx].End;
-			TraceParams.bTraceComplex = (Wheels[WheelIdx]->SweepType == ESweepType::ComplexSweep);
+			TraceParams.bTraceComplex = (WheelTraceParams[WheelIdx].SweepType == ESweepType::ComplexSweep);
 
 			if (GWheeledVehicleDebugParams.TraceTypeOverride > 0)
 			{
@@ -460,7 +458,7 @@ void UChaosWheeledVehicleSimulation::PerformSuspensionTraces(const TArray<FSuspe
 			FVector TraceVector(TraceStart - TraceEnd); // reversed
 			FVector TraceNormal = TraceVector.GetSafeNormal();
 
-			switch (Wheels[WheelIdx]->SweepShape)
+			switch (WheelTraceParams[WheelIdx].SweepShape)
 			{
 			case ESweepShape::Spherecast:
 			{
@@ -492,11 +490,10 @@ void UChaosWheeledVehicleSimulation::PerformSuspensionTraces(const TArray<FSuspe
 
 void UChaosWheeledVehicleSimulation::ApplyWheelFrictionForces(float DeltaTime)
 {
-	for (int WheelIdx = 0; WheelIdx < Wheels.Num(); WheelIdx++)
+	for (int WheelIdx = 0; WheelIdx < PVehicle->Wheels.Num(); WheelIdx++)
 	{
 		auto& PWheel = PVehicle->Wheels[WheelIdx]; // Physics Wheel
-		FHitResult& HitResult = Wheels[WheelIdx]->HitResult;
-		UChaosVehicleWheel* VehicleWheel = Wheels[WheelIdx];
+		FHitResult& HitResult = WheelState.TraceResult[WheelIdx];
 
 		if (PWheel.InContact())
 		{
@@ -506,7 +503,7 @@ void UChaosWheeledVehicleSimulation::ApplyWheelFrictionForces(float DeltaTime)
 			}
 
 			// take into account steering angle
-			float SteerAngleDegrees = VehicleWheel->GetSteerAngle(); 
+			float SteerAngleDegrees = PWheel.SteeringAngle;
 			FRotator SteeringRotator(0.f, SteerAngleDegrees, 0.f);
 			FVector SteerLocalWheelVelocity = SteeringRotator.UnrotateVector(WheelState.LocalWheelVelocity[WheelIdx]);
 
@@ -574,11 +571,11 @@ void UChaosWheeledVehicleSimulation::ApplyWheelFrictionForces(float DeltaTime)
 void UChaosWheeledVehicleSimulation::ApplySuspensionForces(float DeltaTime)
 {
 	TArray<float> SusForces;
-	SusForces.Init(0.f, Wheels.Num());
+	SusForces.Init(0.f, PVehicle->Suspension.Num());
 
-	for (int WheelIdx = 0; WheelIdx < Wheels.Num(); WheelIdx++)
+	for (int WheelIdx = 0; WheelIdx < SusForces.Num(); WheelIdx++)
 	{
-		FHitResult& HitResult = Wheels[WheelIdx]->HitResult;
+		FHitResult& HitResult = WheelState.TraceResult[WheelIdx];
 
 		float NewDesiredLength = 1.0f; // suspension max length
 		float ForceMagnitude2 = 0.f;
@@ -690,10 +687,10 @@ void UChaosWheeledVehicleSimulation::ProcessSteering(const FControlInputs& Contr
 {
 	auto& PSteering = PVehicle->GetSteering();
 
-	for (int WheelIdx = 0; WheelIdx < Wheels.Num(); WheelIdx++)
+	for (int WheelIdx = 0; WheelIdx < PVehicle->Wheels.Num(); WheelIdx++)
 	{
 		auto& PWheel = PVehicle->Wheels[WheelIdx]; // Physics Wheel
-		FHitResult& HitResult = Wheels[WheelIdx]->HitResult;
+		FHitResult& HitResult = WheelState.TraceResult[WheelIdx];
 
 		if (PWheel.SteeringEnabled)
 		{
@@ -769,7 +766,7 @@ void UChaosWheeledVehicleSimulation::ApplyInput(const FControlInputs& ControlInp
 		EngineBraking = PEngine.GetEngineRPM() * PEngine.Setup().EngineBrakeEffect;
 	}
 
-	for (int WheelIdx = 0; WheelIdx < Wheels.Num(); WheelIdx++)
+	for (int WheelIdx = 0; WheelIdx < PVehicle->Wheels.Num(); WheelIdx++)
 	{
 		auto& PWheel = PVehicle->Wheels[WheelIdx];
 
@@ -846,7 +843,7 @@ void UChaosWheeledVehicleSimulation::ProcessMechanicalSimulation(float DeltaTime
 		}
 
 		// apply drive torque to wheels
-		for (int WheelIdx = 0; WheelIdx < Wheels.Num(); WheelIdx++)
+		for (int WheelIdx = 0; WheelIdx < PVehicle->Wheels.Num(); WheelIdx++)
 		{
 			auto& PWheel = PVehicle->Wheels[WheelIdx];
 			if (PWheel.Setup().EngineEnabled)
@@ -910,7 +907,7 @@ void UChaosWheeledVehicleSimulation::DrawDebug3D()
 	{
 		for (int WheelIdx = 0; WheelIdx < PVehicle->Suspension.Num(); WheelIdx++)
 		{
-			FHitResult& Hit = Wheels[WheelIdx]->HitResult;
+			FHitResult& Hit = WheelState.TraceResult[WheelIdx];
 
 			FVector VehicleRightAxis = VehicleState.VehicleWorldTransform.GetUnitAxis(EAxis::Y) * 20.0f;
 			FVector VehicleUpAxis = VehicleState.VehicleWorldTransform.GetUnitAxis(EAxis::Z) * 20.0f;
@@ -932,7 +929,7 @@ void UChaosWheeledVehicleSimulation::DrawDebug3D()
 	{
 		for (int WheelIdx = 0; WheelIdx < PVehicle->Suspension.Num(); WheelIdx++)
 		{
-			FHitResult& Hit = Wheels[WheelIdx]->HitResult;
+			FHitResult& Hit = WheelState.TraceResult[WheelIdx];
 
 			FVector VehicleRightAxis = VehicleState.VehicleWorldTransform.GetUnitAxis(EAxis::Y) * 20.0f;
 			FVector VehicleUpAxis = VehicleState.VehicleWorldTransform.GetUnitAxis(EAxis::Z) * 20.0f;
@@ -954,9 +951,9 @@ void UChaosWheeledVehicleSimulation::DrawDebug3D()
 	if (GWheeledVehicleDebugParams.ShowWheelCollisionNormal)
 	{
 		FString Name;
-		for (int WheelIdx = 0; WheelIdx < Wheels.Num(); WheelIdx++)
+		for (int WheelIdx = 0; WheelIdx < PVehicle->Suspension.Num(); WheelIdx++)
 		{
-			FHitResult& Hit = Wheels[WheelIdx]->HitResult;
+			FHitResult& Hit = WheelState.TraceResult[WheelIdx];
 
 			FVector VehicleRightAxis = VehicleState.VehicleWorldTransform.GetUnitAxis(EAxis::Y) * 20.0f;
 			const FVector& WheelOffset = PVehicle->Suspension[WheelIdx].GetLocalRestingPosition();
@@ -973,7 +970,7 @@ void UChaosWheeledVehicleSimulation::DrawDebug3D()
 
 	if (GWheeledVehicleDebugParams.ShowSuspensionRaycasts)
 	{
-		for (int WheelIdx = 0; WheelIdx < Wheels.Num(); WheelIdx++)
+		for (int WheelIdx = 0; WheelIdx < PVehicle->Suspension.Num(); WheelIdx++)
 		{
 			const FVector& TraceStart = WheelState.Trace[WheelIdx].Start;
 			const FVector& TraceEnd = WheelState.Trace[WheelIdx].End;
@@ -1039,6 +1036,11 @@ void UChaosWheeledVehicleSimulation::FillOutputState(FChaosVehicleAsyncOutput& O
 		WheelsOut.SuspensionOffset = VehicleSuspension[WheelIdx].GetSuspensionOffset();
 		WheelsOut.SpringForce = VehicleSuspension[WheelIdx].GetSuspensionForce();
 		WheelsOut.NormalizedSuspensionLength = VehicleSuspension[WheelIdx].GetNormalizedLength();
+
+		WheelsOut.bBlockingHit = WheelState.TraceResult[WheelIdx].bBlockingHit;
+		WheelsOut.ImpactPoint = WheelState.TraceResult[WheelIdx].ImpactPoint;
+		WheelsOut.PhysMaterial = WheelState.TraceResult[WheelIdx].PhysMaterial;
+
 		Output.VehicleSimOutput.Wheels.Add(WheelsOut);
 	}
 
@@ -1671,6 +1673,13 @@ void UChaosWheeledVehicleMovementComponent::Update(float DeltaTime)
 				TraceParams.bTraceComplex = true;
 				AsyncInput->TraceParams = TraceParams;
 				AsyncInput->TraceCollisionResponse = WheelTraceCollisionResponses;
+
+				AsyncInput->WheelTraceParams.SetNum(Wheels.Num());
+				for (int I = 0; I < Wheels.Num(); I++)
+				{
+					AsyncInput->WheelTraceParams[I].SweepType = Wheels[I]->SweepType;
+					AsyncInput->WheelTraceParams[I].SweepShape = Wheels[I]->SweepShape;
+				}
 			}
 		}
 	}
@@ -2127,17 +2136,16 @@ void UChaosWheeledVehicleMovementComponent::FillWheelOutputState()
 	{
 		if (CurrentOutput->bValid && PVehicleOutput)
 		{
-			for (int WheelIdx = 0; WheelIdx < Wheels.Num(); WheelIdx++)
+			for (int WheelIdx = 0; WheelIdx < WheelStatus.Num(); WheelIdx++)
 			{
 				auto& PWheel = PVehicleOutput->Wheels[WheelIdx];
-				FHitResult& HitResult = Wheels[WheelIdx]->HitResult;
 
 				FWheelStatus& State = WheelStatus[WheelIdx];
 
 				State.bIsValid = true;
-				State.bInContact = HitResult.bBlockingHit;
-				State.ContactPoint = HitResult.ImpactPoint;
-				State.PhysMaterial = HitResult.PhysMaterial;
+				State.bInContact = PWheel.bBlockingHit;
+				State.ContactPoint = PWheel.ImpactPoint;
+				State.PhysMaterial = PWheel.PhysMaterial;
 				State.NormalizedSuspensionLength = PWheel.NormalizedSuspensionLength;
 				State.SpringForce = PWheel.SpringForce;
 				State.SlipAngle = PWheel.SlipAngle;
@@ -2730,6 +2738,13 @@ float UChaosWheeledVehicleMovementComponent::GetSuspensionOffset(int WheelIndex)
 	return Offset;
 }
 
+
+UPhysicalMaterial* UChaosWheeledVehicleMovementComponent::GetPhysMaterial(int WheelIndex)
+{
+	return WheelStatus[WheelIndex].PhysMaterial.Get();
+}
+
+
 FChaosWheelSetup::FChaosWheelSetup()
 	: WheelClass(UChaosVehicleWheel::StaticClass())
 //	, SteeringBoneName(NAME_None)
diff --git a/Engine/Plugins/Experimental/ChaosVehiclesPlugin/Source/ChaosVehicles/Public/ChaosVehicleManagerAsyncCallback.h b/Engine/Plugins/Experimental/ChaosVehiclesPlugin/Source/ChaosVehicles/Public/ChaosVehicleManagerAsyncCallback.h
--- a/Engine/Plugins/Experimental/ChaosVehiclesPlugin/Source/ChaosVehicles/Public/ChaosVehicleManagerAsyncCallback.h
+++ b/Engine/Plugins/Experimental/ChaosVehiclesPlugin/Source/ChaosVehicles/Public/ChaosVehicleManagerAsyncCallback.h
@@ -38,6 +38,9 @@ struct CHAOSVEHICLES_API FWheelsOutput
 		, SuspensionOffset(0.f)
 		, SpringForce(0.f)
 		, NormalizedSuspensionLength(0.f)
+		, bBlockingHit(false)
+		, ImpactPoint(FVector::ZeroVector)
+		, PhysMaterial(nullptr)
 	{
 	}
 
@@ -63,6 +66,9 @@ struct CHAOSVEHICLES_API FWheelsOutput
 	float SpringForce;
 	float NormalizedSuspensionLength;
 
+	bool bBlockingHit;
+	FVector ImpactPoint;
+	TWeakObjectPtr<UPhysicalMaterial> PhysMaterial;
 };
 
 /**
diff --git a/Engine/Plugins/Experimental/ChaosVehiclesPlugin/Source/ChaosVehicles/Public/ChaosVehicleMovementComponent.h b/Engine/Plugins/Experimental/ChaosVehiclesPlugin/Source/ChaosVehicles/Public/ChaosVehicleMovementComponent.h
--- a/Engine/Plugins/Experimental/ChaosVehiclesPlugin/Source/ChaosVehicles/Public/ChaosVehicleMovementComponent.h
+++ b/Engine/Plugins/Experimental/ChaosVehiclesPlugin/Source/ChaosVehicles/Public/ChaosVehicleMovementComponent.h
@@ -164,12 +164,19 @@ struct CHAOSVEHICLES_API FVehicleReplicatedState
 
 };
 
+struct CHAOSVEHICLES_API FWheelTraceParams
+{
+	ESweepType SweepType;
+	ESweepShape SweepShape;
+};
+
 struct CHAOSVEHICLES_API FChaosVehicleDefaultAsyncInput : public FChaosVehicleAsyncInput
 {
 	float GravityZ;
 	FControlInputs ControlInputs;
 	mutable FCollisionQueryParams TraceParams;
 	mutable FCollisionResponseContainer TraceCollisionResponse;
+	mutable TArray<FWheelTraceParams> WheelTraceParams;
 
 	FChaosVehicleDefaultAsyncInput();
 
diff --git a/Engine/Plugins/Experimental/ChaosVehiclesPlugin/Source/ChaosVehicles/Public/ChaosVehicleWheel.h b/Engine/Plugins/Experimental/ChaosVehiclesPlugin/Source/ChaosVehicles/Public/ChaosVehicleWheel.h
--- a/Engine/Plugins/Experimental/ChaosVehiclesPlugin/Source/ChaosVehicles/Public/ChaosVehicleWheel.h
+++ b/Engine/Plugins/Experimental/ChaosVehiclesPlugin/Source/ChaosVehicles/Public/ChaosVehicleWheel.h
@@ -413,8 +413,5 @@ class UChaosWheeledVehicleMovementComponent;
 		/** Get contact surface material */
 		UPhysicalMaterial* GetContactSurfaceMaterial();
 
-		/** suspension raycast results */
-		FHitResult HitResult;
-
 	};
 
diff --git a/Engine/Plugins/Experimental/ChaosVehiclesPlugin/Source/ChaosVehicles/Public/ChaosWheeledVehicleMovementComponent.h b/Engine/Plugins/Experimental/ChaosVehiclesPlugin/Source/ChaosVehicles/Public/ChaosWheeledVehicleMovementComponent.h
--- a/Engine/Plugins/Experimental/ChaosVehiclesPlugin/Source/ChaosVehicles/Public/ChaosWheeledVehicleMovementComponent.h
+++ b/Engine/Plugins/Experimental/ChaosVehiclesPlugin/Source/ChaosVehicles/Public/ChaosWheeledVehicleMovementComponent.h
@@ -523,6 +523,7 @@ struct CHAOSVEHICLES_API FWheelState
 		WorldWheelVelocity.Init(FVector::ZeroVector, NumWheels);
 		LocalWheelVelocity.Init(FVector::ZeroVector, NumWheels);
 		Trace.SetNum(NumWheels);
+		TraceResult.SetNum(NumWheels);
 	}
 
 	/** Commonly used Wheel state - evaluated once used wherever required for that frame */
@@ -536,6 +537,7 @@ struct CHAOSVEHICLES_API FWheelState
 	TArray<FVector> WorldWheelVelocity; /** Current velocity at wheel location In World Coordinates - combined linear and angular */
 	TArray<FVector> LocalWheelVelocity; /** Local velocity of Wheel */
 	TArray<Chaos::FSuspensionTrace> Trace;
+	TArray<FHitResult> TraceResult;
 };
 
 //////////////////////////////////////////////////////////////////////////
@@ -544,8 +546,8 @@ class CHAOSVEHICLES_API UChaosWheeledVehicleSimulation : public UChaosVehicleSim
 {
 public:
 
-	UChaosWheeledVehicleSimulation(TArray<class UChaosVehicleWheel*>& WheelsIn)
-		: Wheels(WheelsIn), bOverlapHit(false)
+	UChaosWheeledVehicleSimulation()
+		: bOverlapHit(false)
 	{
 		QueryBox.Init();
 	}
@@ -577,7 +579,7 @@ public:
 	virtual void ApplyInput(const FControlInputs& ControlInputs, float DeltaTime) override;
 
 	/** Perform suspension ray/shape traces */
-	virtual void PerformSuspensionTraces(const TArray<Chaos::FSuspensionTrace>& SuspensionTrace, FCollisionQueryParams& TraceParams, FCollisionResponseContainer& CollisionResponse);
+	virtual void PerformSuspensionTraces(const TArray<Chaos::FSuspensionTrace>& SuspensionTrace, FCollisionQueryParams& TraceParams, FCollisionResponseContainer& CollisionResponse, TArray<FWheelTraceParams>& WheelTraceParams);
 
 
 	/** Update the engine/transmission simulation */
@@ -599,9 +601,6 @@ public:
 	/** Draw 3D debug lines and things along side the 3D model */
 	virtual void DrawDebug3D() override;
 
-	// #TODO - should not exist here, now duplicated, best not to access at all from Physics 
-	TArray<class UChaosVehicleWheel*>& Wheels;
-
 	FWheelState WheelState;	/** Cached state that holds wheel data for this frame */
 
 	TArray<FPhysicsConstraintHandle> ConstraintHandles;
@@ -704,6 +703,7 @@ class CHAOSVEHICLES_API UChaosWheeledVehicleMovementComponent : public UChaosVeh
 	}
 
 	virtual float GetSuspensionOffset(int WheelIndex) override;
+	UPhysicalMaterial* GetPhysMaterial(int WheelIndex);
 
 	/** Set all channels to the specified response - for wheel raycasts */
 	void SetWheelTraceAllChannels(ECollisionResponse NewResponse)
@@ -843,7 +843,7 @@ class CHAOSVEHICLES_API UChaosWheeledVehicleMovementComponent : public UChaosVeh
 	virtual TUniquePtr<Chaos::FSimpleWheeledVehicle> CreatePhysicsVehicle() override
 	{
 		// Make the Vehicle Simulation class that will be updated from the physics thread async callback
-		VehicleSimulationPT = MakeUnique<UChaosWheeledVehicleSimulation>(Wheels);
+		VehicleSimulationPT = MakeUnique<UChaosWheeledVehicleSimulation>();
 
 		return UChaosVehicleMovementComponent::CreatePhysicsVehicle();
 	}
