Fix CIS warning - Change FChaosVehicleManagerAsyncCallback to use the new options style.

diff --git a/Engine/Plugins/Experimental/ChaosVehiclesPlugin/Source/ChaosVehicles/Private/ChaosVehicleManager.cpp b/Engine/Plugins/Experimental/ChaosVehiclesPlugin/Source/ChaosVehicles/Private/ChaosVehicleManager.cpp
--- a/Engine/Plugins/Experimental/ChaosVehiclesPlugin/Source/ChaosVehicles/Private/ChaosVehicleManager.cpp
+++ b/Engine/Plugins/Experimental/ChaosVehiclesPlugin/Source/ChaosVehicles/Private/ChaosVehicleManager.cpp
@@ -77,7 +77,7 @@ void FChaosVehicleManager::RegisterCallbacks()
 	OnPhysScenePostTickHandle = Scene.OnPhysScenePostTick.AddRaw(this, &FChaosVehicleManager::PostUpdate);
 
 	check(AsyncCallback == nullptr);
-	AsyncCallback = Scene.GetSolver()->CreateAndRegisterSimCallbackObject_External<FChaosVehicleManagerAsyncCallback>(true);
+	AsyncCallback = Scene.GetSolver()->CreateAndRegisterSimCallbackObject_External<FChaosVehicleManagerAsyncCallback>();
 }
 
 void FChaosVehicleManager::UnregisterCallbacks()
diff --git a/Engine/Plugins/Experimental/ChaosVehiclesPlugin/Source/ChaosVehicles/Public/ChaosVehicleManagerAsyncCallback.h b/Engine/Plugins/Experimental/ChaosVehiclesPlugin/Source/ChaosVehicles/Public/ChaosVehicleManagerAsyncCallback.h
--- a/Engine/Plugins/Experimental/ChaosVehiclesPlugin/Source/ChaosVehicles/Public/ChaosVehicleManagerAsyncCallback.h
+++ b/Engine/Plugins/Experimental/ChaosVehiclesPlugin/Source/ChaosVehicles/Public/ChaosVehicleManagerAsyncCallback.h
@@ -171,7 +171,7 @@ struct FChaosVehicleManagerAsyncOutput : public Chaos::FSimCallbackOutput
 /**
  * Async callback from the Physics Engine where we can perform our vehicle simulation
  */
-class CHAOSVEHICLES_API FChaosVehicleManagerAsyncCallback : public Chaos::TSimCallbackObject<FChaosVehicleManagerAsyncInput, FChaosVehicleManagerAsyncOutput>
+class CHAOSVEHICLES_API FChaosVehicleManagerAsyncCallback : public Chaos::TSimCallbackObject<FChaosVehicleManagerAsyncInput, FChaosVehicleManagerAsyncOutput, Chaos::ESimCallbackOptions::Presimulate | Chaos::ESimCallbackOptions::ContactModification>
 {
 private:
 	virtual void OnPreSimulate_Internal() override;
