Removing old redundant metadata on property

diff --git a/Engine/Plugins/Experimental/ChaosVehiclesPlugin/Source/ChaosVehicles/Public/ChaosVehicleWheel.h b/Engine/Plugins/Experimental/ChaosVehiclesPlugin/Source/ChaosVehicles/Public/ChaosVehicleWheel.h
--- a/Engine/Plugins/Experimental/ChaosVehiclesPlugin/Source/ChaosVehicles/Public/ChaosVehicleWheel.h
+++ b/Engine/Plugins/Experimental/ChaosVehiclesPlugin/Source/ChaosVehicles/Public/ChaosVehicleWheel.h
@@ -98,11 +98,11 @@ class UChaosWheeledVehicleMovementComponent;
 		float WheelWidth;
 
 		/** Tyre Cornering Ability */
-		UPROPERTY(EditAnywhere, Category = Wheel, meta = (EditCondition = "bNewWheelSimulation"))
+		UPROPERTY(EditAnywhere, Category = Wheel)
 		float CorneringStiffness;
 
 		/** Friction Force Multiplier */
-		UPROPERTY(EditAnywhere, Category = Wheel, meta = (EditCondition = "bNewWheelSimulation"))
+		UPROPERTY(EditAnywhere, Category = Wheel)
 		float FrictionForceMultiplier;
 
 		/** Wheel Lateral Skid Grip Loss, lower number less grip on skid */
