Chaos Vehicle: Fix tooltip text

diff --git a/Engine/Plugins/Experimental/ChaosVehiclesPlugin/Source/ChaosVehicles/Public/ChaosVehicleWheel.h b/Engine/Plugins/Experimental/ChaosVehiclesPlugin/Source/ChaosVehicles/Public/ChaosVehicleWheel.h
--- a/Engine/Plugins/Experimental/ChaosVehiclesPlugin/Source/ChaosVehicles/Public/ChaosVehicleWheel.h
+++ b/Engine/Plugins/Experimental/ChaosVehiclesPlugin/Source/ChaosVehicles/Public/ChaosVehicleWheel.h
@@ -197,11 +197,11 @@ class UChaosWheeledVehicleMovementComponent;
 		UPROPERTY(EditAnywhere, Category = Suspension, meta = (ClampMin = "0.0", UIMin = "0", ClampMax = "1.0", UIMax = "1"))
 		float RollbarScaling;
 
-		/** Whether wheel suspension considers simple, complex, or both */
+		/** Wheel suspension trace type, defaults to ray trace */
 		UPROPERTY(EditAnywhere, Category = Suspension)
 		ESweepShape SweepShape;
 
-		/** Whether wheel suspension considers simple, complex, or both */
+		/** Whether wheel suspension considers simple, complex */
 		UPROPERTY(EditAnywhere, Category = Suspension)
 		ESweepType SweepType;
 
