Bug fix for lagging vehicle wheels when running with simulation disabled in kinematic mode.
Modifying Mesh->bBlendPhysics after init messed with static mesh internal state.

diff --git a/Engine/Plugins/Experimental/ChaosVehiclesPlugin/Source/ChaosVehicles/Private/WheeledVehiclePawn.cpp b/Engine/Plugins/Experimental/ChaosVehiclesPlugin/Source/ChaosVehicles/Private/WheeledVehiclePawn.cpp
--- a/Engine/Plugins/Experimental/ChaosVehiclesPlugin/Source/ChaosVehicles/Private/WheeledVehiclePawn.cpp
+++ b/Engine/Plugins/Experimental/ChaosVehiclesPlugin/Source/ChaosVehicles/Private/WheeledVehiclePawn.cpp
@@ -18,7 +18,6 @@ AWheeledVehiclePawn::AWheeledVehiclePawn(const FObjectInitializer& ObjectInitial
 	Mesh->BodyInstance.bSimulatePhysics = false;
 	Mesh->BodyInstance.bNotifyRigidBodyCollision = true;
 	Mesh->BodyInstance.bUseCCD = true;
-	Mesh->bBlendPhysics = true;
 	Mesh->SetGenerateOverlapEvents(true);
 	Mesh->SetCanEverAffectNavigation(false);
 	RootComponent = Mesh;
