Reinstate suspension smoothing

diff --git a/Engine/Plugins/Experimental/ChaosVehiclesPlugin/Source/ChaosVehicles/Private/ChaosWheeledVehicleMovementComponent.cpp b/Engine/Plugins/Experimental/ChaosVehiclesPlugin/Source/ChaosVehicles/Private/ChaosWheeledVehicleMovementComponent.cpp
--- a/Engine/Plugins/Experimental/ChaosVehiclesPlugin/Source/ChaosVehicles/Private/ChaosWheeledVehicleMovementComponent.cpp
+++ b/Engine/Plugins/Experimental/ChaosVehiclesPlugin/Source/ChaosVehicles/Private/ChaosWheeledVehicleMovementComponent.cpp
@@ -2695,7 +2695,21 @@ float UChaosWheeledVehicleMovementComponent::GetSuspensionOffset(int WheelIndex)
 				{
 					FVector LocalPos = GetWheelRestingPosition(WheelSetup);
 					FVector LocalHitPoint = VehicleWorldTransform.InverseTransformPosition(WheelStatus[WheelIndex].ContactPoint);
-					Offset = LocalHitPoint.Z - LocalPos.Z + PVehicleOutput->Wheels[WheelIndex].WheelRadius;
+					
+					if (CachedState[WheelIndex].bIsValid)
+					{
+						Offset = CachedState[WheelIndex].WheelOffset;
+						float NewOffset = LocalHitPoint.Z - LocalPos.Z + PVehicleOutput->Wheels[WheelIndex].WheelRadius;
+
+						// interpolate between old and new positions or will just jump to new position if Wheel->SuspensionSmoothing == 0
+						float InterpolationMultiplier = 1.0f - (Wheel->SuspensionSmoothing / 11.0f);
+						float Delta = NewOffset - CachedState[WheelIndex].WheelOffset;
+						Offset += Delta * InterpolationMultiplier;
+					}
+					else
+					{
+						Offset = LocalHitPoint.Z - LocalPos.Z + PVehicleOutput->Wheels[WheelIndex].WheelRadius;
+					}
 					Offset = FMath::Clamp(Offset, -Wheel->SuspensionMaxDrop, Wheel->SuspensionMaxRaise);
 				}
 				else
