-Expose more vehicle data to Blueprint
-Also reinstating a couple of setup options that were previously removed at some point in the past
Neither of these affect the default behaviour of the vehicles.

#UE-175586

diff --git a/Engine/Plugins/Experimental/ChaosVehiclesPlugin/Source/ChaosVehicles/Private/ChaosVehicleMovementComponent.cpp b/Engine/Plugins/Experimental/ChaosVehiclesPlugin/Source/ChaosVehicles/Private/ChaosVehicleMovementComponent.cpp
--- a/Engine/Plugins/Experimental/ChaosVehiclesPlugin/Source/ChaosVehicles/Private/ChaosVehicleMovementComponent.cpp
+++ b/Engine/Plugins/Experimental/ChaosVehiclesPlugin/Source/ChaosVehicles/Private/ChaosVehicleMovementComponent.cpp
@@ -958,16 +958,13 @@ void UChaosVehicleMovementComponent::CalcThrottleBrakeInput(float& ThrottleOut,
 	{
 		if (RawThrottleInput > 0.f)
 		{
-		// Note: Removed this condition to support wheel spinning when rolling backareds with accelerator pressed, rather than braking
-		// Should this case be another checkbox option??
-		//	
-		//	// car moving backwards but player wants to move forwards...
-		//	// if vehicle is moving backwards, then press brake
-		//	if (VehicleState.ForwardSpeed < -WrongDirectionThreshold)
-		//	{
-		//		BrakeOut = 1.0f;
-		//		ThrottleOut = 0.0f;
-		//	}
+			// car moving backwards but player wants to move forwards...
+			// if vehicle is moving backwards, then press brake	
+			if (bThrottleAsBrake && VehicleState.ForwardSpeed < -WrongDirectionThreshold)
+			{
+				BrakeOut = 1.0f;
+				ThrottleOut = 0.0f;
+			}
 
 		}
 		else if (RawBrakeInput > 0.f)
@@ -1812,7 +1809,10 @@ void UChaosVehicleMovementComponent::ParallelUpdate(float DeltaSeconds)
 					PVehicleOutput->Wheels[WheelIdx].SuspensionOffset = FMath::Lerp(Current.SuspensionOffset, Next.SuspensionOffset, OutputInterpAlpha);
 					PVehicleOutput->Wheels[WheelIdx].SpringForce = FMath::Lerp(Current.SpringForce, Next.SpringForce, OutputInterpAlpha);
 					PVehicleOutput->Wheels[WheelIdx].NormalizedSuspensionLength = FMath::Lerp(Current.NormalizedSuspensionLength, Next.NormalizedSuspensionLength, OutputInterpAlpha);
+					PVehicleOutput->Wheels[WheelIdx].DriveTorque = FMath::Lerp(Current.DriveTorque, Next.DriveTorque, OutputInterpAlpha);
+					PVehicleOutput->Wheels[WheelIdx].BrakeTorque = FMath::Lerp(Current.BrakeTorque, Next.BrakeTorque, OutputInterpAlpha);
 
+					PVehicleOutput->Wheels[WheelIdx].bABSActivated = Current.bABSActivated;
 					PVehicleOutput->Wheels[WheelIdx].bBlockingHit = Current.bBlockingHit;
 					PVehicleOutput->Wheels[WheelIdx].ImpactPoint = FMath::Lerp(Current.ImpactPoint, Next.ImpactPoint, OutputInterpAlpha);
 					PVehicleOutput->Wheels[WheelIdx].PhysMaterial = Current.PhysMaterial;
@@ -1847,7 +1847,10 @@ void UChaosVehicleMovementComponent::ParallelUpdate(float DeltaSeconds)
 					PVehicleOutput->Wheels[WheelIdx].SuspensionOffset = Current.SuspensionOffset;
 					PVehicleOutput->Wheels[WheelIdx].SpringForce = Current.SpringForce;
 					PVehicleOutput->Wheels[WheelIdx].NormalizedSuspensionLength = Current.NormalizedSuspensionLength;
+					PVehicleOutput->Wheels[WheelIdx].DriveTorque = Current.DriveTorque;
+					PVehicleOutput->Wheels[WheelIdx].BrakeTorque = Current.BrakeTorque;
 
+					PVehicleOutput->Wheels[WheelIdx].bABSActivated = Current.bABSActivated;
 					PVehicleOutput->Wheels[WheelIdx].bBlockingHit = Current.bBlockingHit;
 					PVehicleOutput->Wheels[WheelIdx].ImpactPoint = Current.ImpactPoint;
 					PVehicleOutput->Wheels[WheelIdx].PhysMaterial = Current.PhysMaterial;
diff --git a/Engine/Plugins/Experimental/ChaosVehiclesPlugin/Source/ChaosVehicles/Private/ChaosVehicleWheel.cpp b/Engine/Plugins/Experimental/ChaosVehiclesPlugin/Source/ChaosVehicles/Private/ChaosVehicleWheel.cpp
--- a/Engine/Plugins/Experimental/ChaosVehiclesPlugin/Source/ChaosVehicles/Private/ChaosVehicleWheel.cpp
+++ b/Engine/Plugins/Experimental/ChaosVehiclesPlugin/Source/ChaosVehicles/Private/ChaosVehicleWheel.cpp
@@ -24,7 +24,7 @@ UChaosVehicleWheel::UChaosVehicleWheel(const FObjectInitializer& ObjectInitializ
 	WheelRadius = 32.0f;
 	WheelWidth = 20.0f;
 	//bAutoAdjustCollisionSize = true;
-	//WheelMass = 20.0f;
+	WheelMass = 20.0f;
 	FrictionForceMultiplier = 2.0f;
 	CorneringStiffness = 1000.0f;
 	SideSlipModifier = 1.0f;
diff --git a/Engine/Plugins/Experimental/ChaosVehiclesPlugin/Source/ChaosVehicles/Private/ChaosWheeledVehicleMovementComponent.cpp b/Engine/Plugins/Experimental/ChaosVehiclesPlugin/Source/ChaosVehicles/Private/ChaosWheeledVehicleMovementComponent.cpp
--- a/Engine/Plugins/Experimental/ChaosVehiclesPlugin/Source/ChaosVehicles/Private/ChaosWheeledVehicleMovementComponent.cpp
+++ b/Engine/Plugins/Experimental/ChaosVehiclesPlugin/Source/ChaosVehicles/Private/ChaosWheeledVehicleMovementComponent.cpp
@@ -88,7 +88,7 @@ FAutoConsoleCommand CVarCommandVehiclesPrevDebugPage(
 
 FString FWheelStatus::ToString() const
 {
-	return FString::Printf(TEXT("bInContact:%s ContactPoint:%s PhysMaterial:%s NormSuspensionLength:%f SpringForce:%f SlipAngle:%f bIsSlipping:%s SlipMagnitude:%f bIsSkidding:%s SkidMagnitude:%f SkidNormal:%s"),
+	return FString::Printf(TEXT("bInContact:%s ContactPoint:%s PhysMaterial:%s NormSuspensionLength:%f SpringForce:%f SlipAngle:%f bIsSlipping:%s SlipMagnitude:%f bIsSkidding:%s SkidMagnitude:%f SkidNormal:%s DriveTorque:%f BrakeTorque:%f ABSActive:%s"),
 		bInContact == true ? TEXT("True") : TEXT("False"),
 		*ContactPoint.ToString(),
 		PhysMaterial.IsValid() ? *PhysMaterial->GetName() : TEXT("None"),
@@ -99,7 +99,10 @@ FString FWheelStatus::ToString() const
 		SlipMagnitude,
 		bIsSkidding == true ? TEXT("True") : TEXT("False"),
 		SkidMagnitude,
-		*SkidNormal.ToString());
+		*SkidNormal.ToString(),
+		DriveTorque,
+		BrakeTorque,
+		bABSActivated == true ? TEXT("True") : TEXT("False"));
 }
 
 void FWheelState::CaptureState(int WheelIdx, const FVector& WheelOffset, const FBodyInstance* TargetInstance)
@@ -1050,6 +1053,10 @@ void UChaosWheeledVehicleSimulation::FillOutputState(FChaosVehicleAsyncOutput& O
 		WheelsOut.SpringForce = VehicleSuspension[WheelIdx].GetSuspensionForce();
 		WheelsOut.NormalizedSuspensionLength = VehicleSuspension[WheelIdx].GetNormalizedLength();
 
+		WheelsOut.DriveTorque = Chaos::TorqueCmToM(VehicleWheels[WheelIdx].GetDriveTorque());
+		WheelsOut.BrakeTorque = Chaos::TorqueCmToM(VehicleWheels[WheelIdx].GetBrakeTorque());
+
+		WheelsOut.bABSActivated = VehicleWheels[WheelIdx].IsABSActivated();
 		WheelsOut.bBlockingHit = WheelState.TraceResult[WheelIdx].bBlockingHit;
 		WheelsOut.ImpactPoint = WheelState.TraceResult[WheelIdx].ImpactPoint;
 		WheelsOut.PhysMaterial = WheelState.TraceResult[WheelIdx].PhysMaterial;
@@ -2182,6 +2189,9 @@ void UChaosWheeledVehicleMovementComponent::FillWheelOutputState()
 				{
 					State.SkidNormal = FVector::ZeroVector;
 				}
+				State.DriveTorque = PWheel.DriveTorque;
+				State.BrakeTorque = PWheel.BrakeTorque;
+				State.bABSActivated = PWheel.bABSActivated;
 			}
 		}
 	}
@@ -2189,7 +2199,7 @@ void UChaosWheeledVehicleMovementComponent::FillWheelOutputState()
 
 
 void UChaosWheeledVehicleMovementComponent::BreakWheelStatus(const struct FWheelStatus& Status, bool& bInContact, FVector& ContactPoint, UPhysicalMaterial*& PhysMaterial
-	, float& NormalizedSuspensionLength, float& SpringForce, float& SlipAngle, bool& bIsSlipping, float& SlipMagnitude, bool& bIsSkidding, float& SkidMagnitude, FVector& SkidNormal)
+	, float& NormalizedSuspensionLength, float& SpringForce, float& SlipAngle, bool& bIsSlipping, float& SlipMagnitude, bool& bIsSkidding, float& SkidMagnitude, FVector& SkidNormal, float& DriveTorque, float& BrakeTorque, bool& bABSActivated)
 {
 	bInContact = Status.bInContact;
 	ContactPoint = Status.ContactPoint;
@@ -2202,10 +2212,13 @@ void UChaosWheeledVehicleMovementComponent::BreakWheelStatus(const struct FWheel
 	bIsSkidding = Status.bIsSkidding;
 	SkidMagnitude = Status.SkidMagnitude;
 	SkidNormal = Status.SkidNormal;
+	DriveTorque = Status.DriveTorque;
+	BrakeTorque = Status.BrakeTorque;
+	bABSActivated = Status.bABSActivated;
 }
 
 FWheelStatus UChaosWheeledVehicleMovementComponent::MakeWheelStatus(bool bInContact, FVector& ContactPoint, UPhysicalMaterial* PhysMaterial
-	, float NormalizedSuspensionLength, float SpringForce, float SlipAngle, bool bIsSlipping, float SlipMagnitude, bool bIsSkidding, float SkidMagnitude, FVector& SkidNormal)
+	, float NormalizedSuspensionLength, float SpringForce, float SlipAngle, bool bIsSlipping, float SlipMagnitude, bool bIsSkidding, float SkidMagnitude, FVector& SkidNormal, float DriveTorque, float BrakeTorque, bool bABSActivated)
 {
 	FWheelStatus Status;
 	Status.bInContact = bInContact;
@@ -2219,6 +2232,9 @@ FWheelStatus UChaosWheeledVehicleMovementComponent::MakeWheelStatus(bool bInCont
 	Status.bIsSkidding = bIsSkidding;
 	Status.SkidMagnitude = SkidMagnitude;
 	Status.SkidNormal = SkidNormal;
+	Status.DriveTorque = DriveTorque;
+	Status.BrakeTorque = BrakeTorque;
+	Status.bABSActivated = bABSActivated;
 
 	return Status;
 }
diff --git a/Engine/Plugins/Experimental/ChaosVehiclesPlugin/Source/ChaosVehicles/Public/ChaosVehicleManagerAsyncCallback.h b/Engine/Plugins/Experimental/ChaosVehiclesPlugin/Source/ChaosVehicles/Public/ChaosVehicleManagerAsyncCallback.h
--- a/Engine/Plugins/Experimental/ChaosVehiclesPlugin/Source/ChaosVehicles/Public/ChaosVehicleManagerAsyncCallback.h
+++ b/Engine/Plugins/Experimental/ChaosVehiclesPlugin/Source/ChaosVehicles/Public/ChaosVehicleManagerAsyncCallback.h
@@ -38,7 +38,9 @@ struct CHAOSVEHICLES_API FWheelsOutput
 		, SuspensionOffset(0.f)
 		, SpringForce(0.f)
 		, NormalizedSuspensionLength(0.f)
-		, bBlockingHit(false)
+		, DriveTorque(0.f)
+		, BrakeTorque(0.f)
+		, bABSActivated(false)
 		, ImpactPoint(FVector::ZeroVector)
 		, PhysMaterial(nullptr)
 	{
@@ -66,6 +68,9 @@ struct CHAOSVEHICLES_API FWheelsOutput
 	float SpringForce;
 	float NormalizedSuspensionLength;
 
+	float DriveTorque;
+	float BrakeTorque;
+	bool bABSActivated;
 	bool bBlockingHit;
 	FVector ImpactPoint;
 	TWeakObjectPtr<UPhysicalMaterial> PhysMaterial;
diff --git a/Engine/Plugins/Experimental/ChaosVehiclesPlugin/Source/ChaosVehicles/Public/ChaosVehicleMovementComponent.h b/Engine/Plugins/Experimental/ChaosVehiclesPlugin/Source/ChaosVehicles/Public/ChaosVehicleMovementComponent.h
--- a/Engine/Plugins/Experimental/ChaosVehiclesPlugin/Source/ChaosVehicles/Public/ChaosVehicleMovementComponent.h
+++ b/Engine/Plugins/Experimental/ChaosVehiclesPlugin/Source/ChaosVehicles/Public/ChaosVehicleMovementComponent.h
@@ -817,6 +817,10 @@ class CHAOSVEHICLES_API UChaosVehicleMovementComponent : public UPawnMovementCom
 	UPROPERTY(EditAnywhere, Category = VehicleSetup)
 	uint8 bReverseAsBrake : 1;
 
+	/** If true, when reversing the throttle will behave like a brake while the vehicle moving in a backwards direction - requires bReverseAsBrake to be enabled for operation */
+	UPROPERTY(EditAnywhere, Category = VehicleSetup, meta = (EditCondition = "bReverseAsBrake"))
+	uint8 bThrottleAsBrake : 1;
+
 public:
 	/** Mass to set the vehicle chassis to. It's much easier to tweak vehicle settings when
 	 * the mass doesn't change due to tweaks with the physics asset. [kg] */
diff --git a/Engine/Plugins/Experimental/ChaosVehiclesPlugin/Source/ChaosVehicles/Public/ChaosVehicleWheel.h b/Engine/Plugins/Experimental/ChaosVehiclesPlugin/Source/ChaosVehicles/Public/ChaosVehicleWheel.h
--- a/Engine/Plugins/Experimental/ChaosVehiclesPlugin/Source/ChaosVehicles/Public/ChaosVehicleWheel.h
+++ b/Engine/Plugins/Experimental/ChaosVehiclesPlugin/Source/ChaosVehicles/Public/ChaosVehicleWheel.h
@@ -97,6 +97,10 @@ class UChaosWheeledVehicleMovementComponent;
 		UPROPERTY(EditAnywhere, Category = Wheel, meta = (ClampMin = "0.01", UIMin = "0.01"))
 		float WheelWidth;
 
+		/** Mass of the wheel Kg */
+		UPROPERTY(EditAnywhere, Category = Wheel, meta = (ClampMin = "0.01", UIMin = "0.01"))
+		float WheelMass;
+
 		/** Tyre Cornering Ability */
 		UPROPERTY(EditAnywhere, Category = Wheel)
 		float CorneringStiffness;
@@ -338,6 +342,7 @@ class UChaosWheeledVehicleMovementComponent;
 			//PWheelConfig.WheelMass = this->WheelMass;
 			PWheelConfig.WheelRadius = this->WheelRadius;
 			PWheelConfig.WheelWidth = this->WheelWidth;
+			PWheelConfig.WheelMass = this->WheelMass;
 			PWheelConfig.MaxSteeringAngle = this->MaxSteerAngle;
 			PWheelConfig.MaxBrakeTorque = this->MaxBrakeTorque;
 			PWheelConfig.HandbrakeTorque = this->MaxHandBrakeTorque;
diff --git a/Engine/Plugins/Experimental/ChaosVehiclesPlugin/Source/ChaosVehicles/Public/ChaosWheeledVehicleMovementComponent.h b/Engine/Plugins/Experimental/ChaosVehiclesPlugin/Source/ChaosVehicles/Public/ChaosWheeledVehicleMovementComponent.h
--- a/Engine/Plugins/Experimental/ChaosVehiclesPlugin/Source/ChaosVehicles/Public/ChaosWheeledVehicleMovementComponent.h
+++ b/Engine/Plugins/Experimental/ChaosVehiclesPlugin/Source/ChaosVehicles/Public/ChaosWheeledVehicleMovementComponent.h
@@ -123,6 +123,18 @@ struct CHAOSVEHICLES_API FWheelStatus
 	UPROPERTY()
 	FVector SkidNormal;
 
+	/** Drive torque currently applied at wheel */
+	UPROPERTY()
+	float DriveTorque;
+
+	/** Brake torque currently applied at wheel */
+	UPROPERTY()
+	float BrakeTorque;
+
+	/** Is the ABS currently engaged - useful for audio Q's */
+	UPROPERTY()
+	bool bABSActivated;
+
 	FWheelStatus()
 	{
 		Init();
@@ -151,6 +163,7 @@ struct CHAOSVEHICLES_API FWheelStatus
 		SkidNormal = FVector::ZeroVector;
 		ContactPoint = FVector::ZeroVector;
 		bIsValid = false;
+		bABSActivated = false;
 	}
 
 	FString ToString() const;
@@ -673,11 +686,11 @@ class CHAOSVEHICLES_API UChaosWheeledVehicleMovementComponent : public UChaosVeh
 
 	UFUNCTION(BlueprintPure, Category = "Vehicles")
 	static void BreakWheelStatus(const struct FWheelStatus& Status, bool& bInContact, FVector& ContactPoint, UPhysicalMaterial*& PhysMaterial
-			, float& NormalizedSuspensionLength, float& SpringForce, float& SlipAngle, bool& bIsSlipping, float& SlipMagnitude, bool& bIsSkidding, float& SkidMagnitude, FVector& SkidNormal);
+			, float& NormalizedSuspensionLength, float& SpringForce, float& SlipAngle, bool& bIsSlipping, float& SlipMagnitude, bool& bIsSkidding, float& SkidMagnitude, FVector& SkidNormal, float& DriveTorque, float& BrakeTorque, bool& bABSActivated);
 
 	UFUNCTION(BlueprintPure, Category = "Vehicles")
 	static FWheelStatus MakeWheelStatus(bool bInContact, FVector& ContactPoint, UPhysicalMaterial* PhysMaterial
-			, float NormalizedSuspensionLength, float SpringForce, float SlipAngle, bool bIsSlipping, float SlipMagnitude, bool bIsSkidding, float SkidMagnitude, FVector& SkidNormal);
+			, float NormalizedSuspensionLength, float SpringForce, float SlipAngle, bool bIsSlipping, float SlipMagnitude, bool bIsSkidding, float SkidMagnitude, FVector& SkidNormal, float DriveTorque, float BrakeTorque, bool bABSActivated);
 
 	UFUNCTION(BlueprintPure, Category = "Vehicles")
 	static void BreakWheeledSnapshot(const struct FWheeledSnaphotData& Snapshot, FTransform& Transform, FVector& LinearVelocity
diff --git a/Engine/Source/Runtime/Experimental/ChaosVehicles/ChaosVehiclesCore/Private/WheelSystem.cpp b/Engine/Source/Runtime/Experimental/ChaosVehicles/ChaosVehiclesCore/Private/WheelSystem.cpp
--- a/Engine/Source/Runtime/Experimental/ChaosVehicles/ChaosVehiclesCore/Private/WheelSystem.cpp
+++ b/Engine/Source/Runtime/Experimental/ChaosVehicles/ChaosVehiclesCore/Private/WheelSystem.cpp
@@ -46,6 +46,7 @@ namespace Chaos
 		, AvailableGrip(0.f)
 		, InputForces(FVector::ZeroVector)
 		, bClipping(false)
+		, bABSActivated(false)
 	{
 
 	}
@@ -84,11 +85,13 @@ namespace Chaos
 
 		float FinalLongitudinalForce = 0.f;
 		float FinalLateralForce = 0.f;
+		float ABSGripThresholdSpeed = 5.0f;
 
 		// currently just letting the brake override the throttle
 		bool Braking = BrakeTorque > FMath::Abs(DriveTorque);
 		bool WheelLocked = false;
 		float SlipOmega = 0.0f;
+		bABSActivated = 0.0f;
 
 		// are we actually touching the ground
 		if (ForceIntoSurface > SMALL_NUMBER)
@@ -100,6 +103,10 @@ namespace Chaos
 				{
 					float Sign = (AppliedLinearBrakeForce > 0.0f) ? 1.0f : -1.0f;
 					AppliedLinearBrakeForce = AvailableGrip * TractionControlAndABSScaling * Sign;
+					if (FMath::Abs(GroundVelocityVector.X) > ABSGripThresholdSpeed)
+					{
+						bABSActivated = true;
+					}
 				}
 			}
 
@@ -201,8 +208,15 @@ namespace Chaos
 		}
 		else
 		{ 
-			float GroundOmega = GroundVelocityVector.X / FMath::Max(Re, KINDA_SMALL_NUMBER);
-			Omega += ((GroundOmega - Omega + SlipOmega));
+			if (bInContact)
+			{
+				float GroundOmega = GroundVelocityVector.X / FMath::Max(Re, KINDA_SMALL_NUMBER);
+				Omega += ((GroundOmega - Omega + SlipOmega));
+			}
+			else
+			{
+				Omega += SlipOmega;
+			}
 		}
 
 		// Wheel angular position integrated
diff --git a/Engine/Source/Runtime/Experimental/ChaosVehicles/ChaosVehiclesCore/Public/VehicleUtility.h b/Engine/Source/Runtime/Experimental/ChaosVehicles/ChaosVehiclesCore/Public/VehicleUtility.h
--- a/Engine/Source/Runtime/Experimental/ChaosVehicles/ChaosVehiclesCore/Public/VehicleUtility.h
+++ b/Engine/Source/Runtime/Experimental/ChaosVehicles/ChaosVehiclesCore/Public/VehicleUtility.h
@@ -275,6 +275,11 @@ namespace Chaos
 		return TorqueIn * 10000.0f;
 	}
 
+	FORCEINLINE float TorqueCmToM(float TorqueIn)
+	{
+		return TorqueIn / 10000.0f;
+	}
+
 	class CHAOSVEHICLESCORE_API FTimeAndDistanceMeasure
 	{
 	public:
diff --git a/Engine/Source/Runtime/Experimental/ChaosVehicles/ChaosVehiclesCore/Public/WheelSystem.h b/Engine/Source/Runtime/Experimental/ChaosVehicles/ChaosVehiclesCore/Public/WheelSystem.h
--- a/Engine/Source/Runtime/Experimental/ChaosVehicles/ChaosVehiclesCore/Public/WheelSystem.h
+++ b/Engine/Source/Runtime/Experimental/ChaosVehicles/ChaosVehiclesCore/Public/WheelSystem.h
@@ -325,6 +325,10 @@ public:
 	{
 		return (FMath::Abs(GetSkidMagnitude()) > Setup().SkidThreshold);
 	}
+	bool IsABSActivated() const
+	{
+		return bABSActivated;
+	}
 
 	float GetSteeringAngle() const
 	{
@@ -467,6 +471,8 @@ public:
 	float AvailableGrip;
 	FVector InputForces;
 	bool bClipping;
+	bool bABSActivated;
+
 };
 
 
