Adding ability to modify Suspension Settings dynamically at runtime. This is required for project Opal and would be useful for all users.

diff --git a/Engine/Plugins/Experimental/ChaosVehiclesPlugin/Source/ChaosVehicles/Private/ChaosWheeledVehicleMovementComponent.cpp b/Engine/Plugins/Experimental/ChaosVehiclesPlugin/Source/ChaosVehicles/Private/ChaosWheeledVehicleMovementComponent.cpp
--- a/Engine/Plugins/Experimental/ChaosVehiclesPlugin/Source/ChaosVehicles/Private/ChaosWheeledVehicleMovementComponent.cpp
+++ b/Engine/Plugins/Experimental/ChaosVehiclesPlugin/Source/ChaosVehicles/Private/ChaosWheeledVehicleMovementComponent.cpp
@@ -2694,6 +2694,50 @@ void UChaosWheeledVehicleMovementComponent::SetBrakeTorque(float BrakeTorque, in
 	}
 }
 
+void UChaosWheeledVehicleMovementComponent::SetSuspensionParams(float Rate, float Damping, float Preload, float MaxRaise, float MaxDrop, int32 WheelIndex)
+{
+	using namespace Chaos;
+
+	if (FBodyInstance* TargetInstance = GetBodyInstance())
+	{
+		FPhysicsCommand::ExecuteWrite(TargetInstance->ActorHandle, [&](const FPhysicsActorHandle& Chassis)
+			{
+				if (VehicleSimulationPT && VehicleSimulationPT->PVehicle && WheelIndex < VehicleSimulationPT->PVehicle->Wheels.Num())
+				{
+					Chaos::FSimpleSuspensionSim& VehicleSuspension = VehicleSimulationPT->PVehicle->Suspension[WheelIndex];
+
+					VehicleSuspension.AccessSetup().SpringRate = Rate;
+					VehicleSuspension.AccessSetup().DampingRatio = Damping;
+					VehicleSuspension.AccessSetup().SpringPreload = Preload;
+					VehicleSuspension.AccessSetup().SetSuspensionMaxRaise(MaxRaise);
+					VehicleSuspension.AccessSetup().SetSuspensionMaxDrop(MaxDrop);
+				}
+			});
+
+		FPhysicsCommand::ExecuteWrite(TargetInstance->ActorHandle, [&](const FPhysicsActorHandle& Chassis)
+			{
+				FChaosWheelSetup& WheelSetup = WheelSetups[WheelIndex];
+				const FVector LocalWheel = GetWheelRestingPosition(WheelSetup);
+				FPhysicsConstraintHandle ConstraintHandle = FPhysicsInterface::CreateSuspension(Chassis, LocalWheel);
+
+				if (ConstraintHandle.IsValid())
+				{
+					UChaosVehicleWheel* Wheel = Wheels[WheelIndex];
+					check(Wheel);
+					if (Chaos::FSuspensionConstraint* Constraint = static_cast<Chaos::FSuspensionConstraint*>(ConstraintHandles[WheelIndex].Constraint))
+					{
+						Constraint->SetHardstopStiffness(1.0f);
+						Constraint->SetSpringStiffness(Chaos::MToCm(Rate) * 0.25f);
+						Constraint->SetSpringPreload(Chaos::MToCm(Preload));
+						Constraint->SetSpringDamping(Damping * 5.0f);
+						Constraint->SetMinLength(-MaxRaise);
+						Constraint->SetMaxLength(MaxDrop);
+					}
+				}
+			});
+
+	}
+}
 
 float UChaosWheeledVehicleMovementComponent::GetSuspensionOffset(int WheelIndex)
 {
diff --git a/Engine/Plugins/Experimental/ChaosVehiclesPlugin/Source/ChaosVehicles/Public/ChaosWheeledVehicleMovementComponent.h b/Engine/Plugins/Experimental/ChaosVehiclesPlugin/Source/ChaosVehicles/Public/ChaosWheeledVehicleMovementComponent.h
--- a/Engine/Plugins/Experimental/ChaosVehiclesPlugin/Source/ChaosVehicles/Public/ChaosWheeledVehicleMovementComponent.h
+++ b/Engine/Plugins/Experimental/ChaosVehiclesPlugin/Source/ChaosVehicles/Public/ChaosWheeledVehicleMovementComponent.h
@@ -852,6 +852,9 @@ class CHAOSVEHICLES_API UChaosWheeledVehicleMovementComponent : public UChaosVeh
 	UFUNCTION(BlueprintCallable, Category = "Game|Components|ChaosWheeledVehicleMovement")
 	void SetBrakeTorque(float BrakeTorque, int32 WheelIndex);
 
+	UFUNCTION(BlueprintCallable, Category = "Game|Components|ChaosWheeledVehicleMovement")
+	void SetSuspensionParams(float Rate, float Damping, float Preload, float MaxRaise, float MaxDrop, int32 WheelIndex);
+
 	/** */
 	virtual TUniquePtr<Chaos::FSimpleWheeledVehicle> CreatePhysicsVehicle() override
 	{
diff --git a/Engine/Source/Runtime/Experimental/ChaosVehicles/ChaosVehiclesCore/Public/SuspensionSystem.h b/Engine/Source/Runtime/Experimental/ChaosVehicles/ChaosVehiclesCore/Public/SuspensionSystem.h
--- a/Engine/Source/Runtime/Experimental/ChaosVehicles/ChaosVehiclesCore/Public/SuspensionSystem.h
+++ b/Engine/Source/Runtime/Experimental/ChaosVehicles/ChaosVehiclesCore/Public/SuspensionSystem.h
@@ -47,6 +47,18 @@ namespace Chaos
 			SuspensionSmoothing = FMath::Clamp(SuspensionSmoothing, 0, NUM_SUS_AVERAGING);
 		}
 
+		void SetSuspensionMaxRaise(float InSuspensionMaxRaise)
+		{
+			SuspensionMaxRaise = InSuspensionMaxRaise;
+			MaxLength = FMath::Abs(SuspensionMaxRaise) + FMath::Abs(SuspensionMaxDrop);
+		}
+
+		void SetSuspensionMaxDrop(float InSuspensionMaxDrop)
+		{
+			SuspensionMaxDrop = InSuspensionMaxDrop;
+			MaxLength = FMath::Abs(SuspensionMaxRaise) + FMath::Abs(SuspensionMaxDrop);
+		}
+
 		FVector SuspensionAxis;		// local axis, direction of suspension force raycast traces
 		FVector SuspensionForceOffset;	// relative position from wheel where suspension forces are applied
 		float SuspensionMaxRaise;	// distance [cm]
