Possible for code to crash when BP function accessing new vehicle suspension functionality. Needing to check array is not empty in code.

diff --git a/Engine/Plugins/Experimental/ChaosVehiclesPlugin/Source/ChaosVehicles/Private/ChaosWheeledVehicleMovementComponent.cpp b/Engine/Plugins/Experimental/ChaosVehiclesPlugin/Source/ChaosVehicles/Private/ChaosWheeledVehicleMovementComponent.cpp
--- a/Engine/Plugins/Experimental/ChaosVehiclesPlugin/Source/ChaosVehicles/Private/ChaosWheeledVehicleMovementComponent.cpp
+++ b/Engine/Plugins/Experimental/ChaosVehiclesPlugin/Source/ChaosVehicles/Private/ChaosWheeledVehicleMovementComponent.cpp
@@ -2716,22 +2716,25 @@ void UChaosWheeledVehicleMovementComponent::SetSuspensionParams(float Rate, floa
 
 		FPhysicsCommand::ExecuteWrite(TargetInstance->ActorHandle, [&](const FPhysicsActorHandle& Chassis)
 			{
-				FChaosWheelSetup& WheelSetup = WheelSetups[WheelIndex];
-				const FVector LocalWheel = GetWheelRestingPosition(WheelSetup);
-				FPhysicsConstraintHandle ConstraintHandle = FPhysicsInterface::CreateSuspension(Chassis, LocalWheel);
-
-				if (ConstraintHandle.IsValid())
+				if (WheelIndex < WheelSetups.Num() && WheelIndex < Wheels.Num())
 				{
-					UChaosVehicleWheel* Wheel = Wheels[WheelIndex];
-					check(Wheel);
-					if (Chaos::FSuspensionConstraint* Constraint = static_cast<Chaos::FSuspensionConstraint*>(ConstraintHandles[WheelIndex].Constraint))
+					FChaosWheelSetup& WheelSetup = WheelSetups[WheelIndex];
+					const FVector LocalWheel = GetWheelRestingPosition(WheelSetup);
+					FPhysicsConstraintHandle ConstraintHandle = FPhysicsInterface::CreateSuspension(Chassis, LocalWheel);
+
+					if (ConstraintHandle.IsValid())
 					{
-						Constraint->SetHardstopStiffness(1.0f);
-						Constraint->SetSpringStiffness(Chaos::MToCm(Rate) * 0.25f);
-						Constraint->SetSpringPreload(Chaos::MToCm(Preload));
-						Constraint->SetSpringDamping(Damping * 5.0f);
-						Constraint->SetMinLength(-MaxRaise);
-						Constraint->SetMaxLength(MaxDrop);
+						UChaosVehicleWheel* Wheel = Wheels[WheelIndex];
+						check(Wheel);
+						if (Chaos::FSuspensionConstraint* Constraint = static_cast<Chaos::FSuspensionConstraint*>(ConstraintHandles[WheelIndex].Constraint))
+						{
+							Constraint->SetHardstopStiffness(1.0f);
+							Constraint->SetSpringStiffness(Chaos::MToCm(Rate) * 0.25f);
+							Constraint->SetSpringPreload(Chaos::MToCm(Preload));
+							Constraint->SetSpringDamping(Damping * 5.0f);
+							Constraint->SetMinLength(-MaxRaise);
+							Constraint->SetMaxLength(MaxDrop);
+						}
 					}
 				}
 			});
