Opal - Fix suspension jitter issue that was apparent when driving over small obstacles; due to calculation error when using spherecast trace type

diff --git a/Engine/Plugins/Experimental/ChaosVehiclesPlugin/Source/ChaosVehicles/Private/ChaosWheeledVehicleMovementComponent.cpp b/Engine/Plugins/Experimental/ChaosVehiclesPlugin/Source/ChaosVehicles/Private/ChaosWheeledVehicleMovementComponent.cpp
--- a/Engine/Plugins/Experimental/ChaosVehiclesPlugin/Source/ChaosVehicles/Private/ChaosWheeledVehicleMovementComponent.cpp
+++ b/Engine/Plugins/Experimental/ChaosVehiclesPlugin/Source/ChaosVehicles/Private/ChaosWheeledVehicleMovementComponent.cpp
@@ -274,7 +274,7 @@ void UChaosWheeledVehicleSimulation::UpdateSimulation(float DeltaTime, const FCh
 
 		if (!GWheeledVehicleDebugParams.DisableSuspensionForces && PVehicle->bSuspensionEnabled)
 		{
-			ApplySuspensionForces(DeltaTime);
+			ApplySuspensionForces(DeltaTime, InputData.WheelTraceParams);
 		}
 
 		///////////////////////////////////////////////////////////////////////
@@ -574,7 +574,7 @@ void UChaosWheeledVehicleSimulation::ApplyWheelFrictionForces(float DeltaTime)
 	}
 }
 
-void UChaosWheeledVehicleSimulation::ApplySuspensionForces(float DeltaTime)
+void UChaosWheeledVehicleSimulation::ApplySuspensionForces(float DeltaTime, TArray<FWheelTraceParams>& WheelTraceParams)
 {
 	using namespace Chaos;
 
@@ -602,7 +602,15 @@ void UChaosWheeledVehicleSimulation::ApplySuspensionForces(float DeltaTime)
 					{
 						if (FSuspensionConstraintPhysicsProxy* Proxy = Constraint->GetProxy<FSuspensionConstraintPhysicsProxy>())
 						{
-							const FVec3 TargetPos = HitResult.ImpactPoint + (PWheel.GetEffectiveRadius() * VehicleState.VehicleUpAxis);
+							FVec3 TargetPos;
+							if (WheelTraceParams[WheelIdx].SweepShape == ESweepShape::Raycast)
+							{
+								TargetPos = HitResult.ImpactPoint + (PWheel.GetEffectiveRadius() * VehicleState.VehicleUpAxis);
+							}
+							else
+							{
+								TargetPos = HitResult.Location;
+							}
 
 							Chaos::FPhysicsSolver* Solver = Proxy->GetSolver<Chaos::FPhysicsSolver>();
 
diff --git a/Engine/Plugins/Experimental/ChaosVehiclesPlugin/Source/ChaosVehicles/Public/ChaosWheeledVehicleMovementComponent.h b/Engine/Plugins/Experimental/ChaosVehiclesPlugin/Source/ChaosVehicles/Public/ChaosWheeledVehicleMovementComponent.h
--- a/Engine/Plugins/Experimental/ChaosVehiclesPlugin/Source/ChaosVehicles/Public/ChaosWheeledVehicleMovementComponent.h
+++ b/Engine/Plugins/Experimental/ChaosVehiclesPlugin/Source/ChaosVehicles/Public/ChaosWheeledVehicleMovementComponent.h
@@ -605,7 +605,7 @@ public:
 	virtual void ApplyWheelFrictionForces(float DeltaTime);
 
 	/** calculate and apply chassis suspension forces */
-	virtual void ApplySuspensionForces(float DeltaTime);
+	virtual void ApplySuspensionForces(float DeltaTime, TArray<FWheelTraceParams>& WheelTraceParams);
 
 	bool IsWheelSpinning() const;
 	bool ContainsTraces(const FBox& Box, const TArray<struct Chaos::FSuspensionTrace>& SuspensionTrace);
