Fix uninitialised properties!

diff --git a/Engine/Plugins/Experimental/ChaosVehiclesPlugin/Source/ChaosVehicles/Public/ChaosWheeledVehicleMovementComponent.h b/Engine/Plugins/Experimental/ChaosVehiclesPlugin/Source/ChaosVehicles/Public/ChaosWheeledVehicleMovementComponent.h
--- a/Engine/Plugins/Experimental/ChaosVehiclesPlugin/Source/ChaosVehicles/Public/ChaosWheeledVehicleMovementComponent.h
+++ b/Engine/Plugins/Experimental/ChaosVehiclesPlugin/Source/ChaosVehicles/Public/ChaosWheeledVehicleMovementComponent.h
@@ -164,6 +164,8 @@ struct CHAOSVEHICLES_API FWheelStatus
 		ContactPoint = FVector::ZeroVector;
 		bIsValid = false;
 		bABSActivated = false;
+		DriveTorque = 0.0f;
+		BrakeTorque = 0.0f;
 	}
 
 	FString ToString() const;
