Opal - Further fix for twitchy suspension when travelling over small rocks

diff --git a/Engine/Plugins/Experimental/ChaosVehiclesPlugin/Source/ChaosVehicles/Private/ChaosVehicleMovementComponent.cpp b/Engine/Plugins/Experimental/ChaosVehiclesPlugin/Source/ChaosVehicles/Private/ChaosVehicleMovementComponent.cpp
--- a/Engine/Plugins/Experimental/ChaosVehiclesPlugin/Source/ChaosVehicles/Private/ChaosVehicleMovementComponent.cpp
+++ b/Engine/Plugins/Experimental/ChaosVehiclesPlugin/Source/ChaosVehicles/Private/ChaosVehicleMovementComponent.cpp
@@ -1815,6 +1815,7 @@ void UChaosVehicleMovementComponent::ParallelUpdate(float DeltaSeconds)
 					PVehicleOutput->Wheels[WheelIdx].bABSActivated = Current.bABSActivated;
 					PVehicleOutput->Wheels[WheelIdx].bBlockingHit = Current.bBlockingHit;
 					PVehicleOutput->Wheels[WheelIdx].ImpactPoint = FMath::Lerp(Current.ImpactPoint, Next.ImpactPoint, OutputInterpAlpha);
+					PVehicleOutput->Wheels[WheelIdx].HitLocation = FMath::Lerp(Current.HitLocation, Next.HitLocation, OutputInterpAlpha);
 					PVehicleOutput->Wheels[WheelIdx].PhysMaterial = Current.PhysMaterial;
 				}
 			}
@@ -1853,6 +1854,7 @@ void UChaosVehicleMovementComponent::ParallelUpdate(float DeltaSeconds)
 					PVehicleOutput->Wheels[WheelIdx].bABSActivated = Current.bABSActivated;
 					PVehicleOutput->Wheels[WheelIdx].bBlockingHit = Current.bBlockingHit;
 					PVehicleOutput->Wheels[WheelIdx].ImpactPoint = Current.ImpactPoint;
+					PVehicleOutput->Wheels[WheelIdx].HitLocation = Current.HitLocation;
 					PVehicleOutput->Wheels[WheelIdx].PhysMaterial = Current.PhysMaterial;
 				}
 
diff --git a/Engine/Plugins/Experimental/ChaosVehiclesPlugin/Source/ChaosVehicles/Private/ChaosWheeledVehicleMovementComponent.cpp b/Engine/Plugins/Experimental/ChaosVehiclesPlugin/Source/ChaosVehicles/Private/ChaosWheeledVehicleMovementComponent.cpp
--- a/Engine/Plugins/Experimental/ChaosVehiclesPlugin/Source/ChaosVehicles/Private/ChaosWheeledVehicleMovementComponent.cpp
+++ b/Engine/Plugins/Experimental/ChaosVehiclesPlugin/Source/ChaosVehicles/Private/ChaosWheeledVehicleMovementComponent.cpp
@@ -1067,6 +1067,7 @@ void UChaosWheeledVehicleSimulation::FillOutputState(FChaosVehicleAsyncOutput& O
 		WheelsOut.bABSActivated = VehicleWheels[WheelIdx].IsABSActivated();
 		WheelsOut.bBlockingHit = WheelState.TraceResult[WheelIdx].bBlockingHit;
 		WheelsOut.ImpactPoint = WheelState.TraceResult[WheelIdx].ImpactPoint;
+		WheelsOut.HitLocation = WheelState.TraceResult[WheelIdx].Location;
 		WheelsOut.PhysMaterial = WheelState.TraceResult[WheelIdx].PhysMaterial;
 
 		Output.VehicleSimOutput.Wheels.Add(WheelsOut);
@@ -2177,6 +2178,7 @@ void UChaosWheeledVehicleMovementComponent::FillWheelOutputState()
 				State.bIsValid = true;
 				State.bInContact = PWheel.bBlockingHit;
 				State.ContactPoint = PWheel.ImpactPoint;
+				State.HitLocation = PWheel.HitLocation;
 				State.PhysMaterial = PWheel.PhysMaterial;
 				State.NormalizedSuspensionLength = PWheel.NormalizedSuspensionLength;
 				State.SpringForce = PWheel.SpringForce;
@@ -2754,6 +2756,20 @@ float UChaosWheeledVehicleMovementComponent::GetSuspensionOffset(int WheelIndex)
 {
 	float Offset = 0.f;
 
+	auto CalcOffset = [this](ESweepShape& SweepShape, FVector& LocalHitPoint, FVector& LocalPos, float Radius)
+	{
+		float NewOffset = 0.0f;
+		if (SweepShape == ESweepShape::Spherecast)
+		{
+			NewOffset = LocalHitPoint.Z - LocalPos.Z;
+		}
+		else
+		{
+			NewOffset = LocalHitPoint.Z - LocalPos.Z + Radius;
+		}
+		return NewOffset;
+	};
+
 	FChaosWheelSetup& WheelSetup = WheelSetups[WheelIndex];
 	if (GetBodyInstance())
 	{
@@ -2765,13 +2781,17 @@ float UChaosWheeledVehicleMovementComponent::GetSuspensionOffset(int WheelIndex)
 				if (WheelStatus[WheelIndex].bInContact)
 				{
 					FVector LocalPos = GetWheelRestingPosition(WheelSetup);
-					FVector LocalHitPoint = VehicleWorldTransform.InverseTransformPosition(WheelStatus[WheelIndex].ContactPoint);
-					
+
+					FVector ReferencePos = (Wheel->SweepShape == ESweepShape::Spherecast)? WheelStatus[WheelIndex].HitLocation : WheelStatus[WheelIndex].ContactPoint;
+					FVector LocalHitPoint = VehicleWorldTransform.InverseTransformPosition(ReferencePos);
+					float Radius = PVehicleOutput->Wheels[WheelIndex].WheelRadius;
+
 					if (CachedState[WheelIndex].bIsValid)
 					{
 						Offset = CachedState[WheelIndex].WheelOffset;
-						float NewOffset = LocalHitPoint.Z - LocalPos.Z + PVehicleOutput->Wheels[WheelIndex].WheelRadius;
 
+						float NewOffset = CalcOffset(Wheel->SweepShape, LocalHitPoint, LocalPos, Radius);
+							
 						// interpolate between old and new positions or will just jump to new position if Wheel->SuspensionSmoothing == 0
 						float InterpolationMultiplier = 1.0f - (Wheel->SuspensionSmoothing / 11.0f);
 						float Delta = NewOffset - CachedState[WheelIndex].WheelOffset;
@@ -2779,7 +2799,7 @@ float UChaosWheeledVehicleMovementComponent::GetSuspensionOffset(int WheelIndex)
 					}
 					else
 					{
-						Offset = LocalHitPoint.Z - LocalPos.Z + PVehicleOutput->Wheels[WheelIndex].WheelRadius;
+						Offset = CalcOffset(Wheel->SweepShape, LocalHitPoint, LocalPos, Radius);
 					}
 					Offset = FMath::Clamp(Offset, -Wheel->SuspensionMaxDrop, Wheel->SuspensionMaxRaise);
 				}
@@ -2820,13 +2840,43 @@ float UChaosWheeledVehicleMovementComponent::GetSuspensionOffset(int WheelIndex)
 					FVector TraceEnd = WorldLocation + WorldDirection * (Wheel->SuspensionMaxDrop + Wheel->WheelRadius);
 
 					FHitResult HitResult;
-					GetWorld()->LineTraceSingleByChannel(HitResult, TraceStart, TraceEnd, SpringCollisionChannel, TraceParams, ResponseParams);
+					switch (Wheel->SweepShape)
+					{
+						case ESweepShape::Spherecast:
+						{
+							float WheelRadius = Wheel->WheelRadius;
+
+							GetWorld()->SweepSingleByChannel(HitResult
+								, TraceStart
+								, TraceEnd
+								, FQuat::Identity, SpringCollisionChannel
+								, FCollisionShape::MakeSphere(WheelRadius), TraceParams
+								, ResponseParams);
+						}
+						break;
+
+						case ESweepShape::Raycast:
+						default:
+						{
+							GetWorld()->LineTraceSingleByChannel(HitResult, TraceStart, TraceEnd, SpringCollisionChannel, TraceParams, ResponseParams);
+						}
+						break;
+					}
 
 					if (HitResult.bBlockingHit)
 					{
 						FVector LocalPos = GetWheelRestingPosition(WheelSetup);
-						FVector LocalHitPoint = VehicleWorldTransform.InverseTransformPosition(HitResult.ImpactPoint);
-						Offset = LocalHitPoint.Z - LocalPos.Z + Wheel->WheelRadius;
+						if (Wheel->SweepShape == ESweepShape::Raycast)
+						{
+							FVector LocalHitPoint = VehicleWorldTransform.InverseTransformPosition(HitResult.ImpactPoint);
+							Offset = CalcOffset(Wheel->SweepShape, LocalHitPoint, LocalPos, Wheel->WheelRadius);
+						}
+						else
+						{
+							FVector LocalHitPoint = VehicleWorldTransform.InverseTransformPosition(HitResult.Location);
+							Offset = CalcOffset(Wheel->SweepShape, LocalHitPoint, LocalPos, Wheel->WheelRadius);
+						}
+
 						Offset = FMath::Clamp(Offset, -Wheel->SuspensionMaxDrop, Wheel->SuspensionMaxRaise);
 					}
 					else
@@ -2840,7 +2890,6 @@ float UChaosWheeledVehicleMovementComponent::GetSuspensionOffset(int WheelIndex)
 			}
 		}
 	}
-	
 	return Offset;
 }
 
diff --git a/Engine/Plugins/Experimental/ChaosVehiclesPlugin/Source/ChaosVehicles/Public/ChaosVehicleManagerAsyncCallback.h b/Engine/Plugins/Experimental/ChaosVehiclesPlugin/Source/ChaosVehicles/Public/ChaosVehicleManagerAsyncCallback.h
--- a/Engine/Plugins/Experimental/ChaosVehiclesPlugin/Source/ChaosVehicles/Public/ChaosVehicleManagerAsyncCallback.h
+++ b/Engine/Plugins/Experimental/ChaosVehiclesPlugin/Source/ChaosVehicles/Public/ChaosVehicleManagerAsyncCallback.h
@@ -42,6 +42,7 @@ struct CHAOSVEHICLES_API FWheelsOutput
 		, BrakeTorque(0.f)
 		, bABSActivated(false)
 		, ImpactPoint(FVector::ZeroVector)
+		, HitLocation(FVector::ZeroVector)
 		, PhysMaterial(nullptr)
 	{
 	}
@@ -73,6 +74,7 @@ struct CHAOSVEHICLES_API FWheelsOutput
 	bool bABSActivated;
 	bool bBlockingHit;
 	FVector ImpactPoint;
+	FVector HitLocation;
 	TWeakObjectPtr<UPhysicalMaterial> PhysMaterial;
 };
 
diff --git a/Engine/Plugins/Experimental/ChaosVehiclesPlugin/Source/ChaosVehicles/Public/ChaosWheeledVehicleMovementComponent.h b/Engine/Plugins/Experimental/ChaosVehiclesPlugin/Source/ChaosVehicles/Public/ChaosWheeledVehicleMovementComponent.h
--- a/Engine/Plugins/Experimental/ChaosVehiclesPlugin/Source/ChaosVehicles/Public/ChaosWheeledVehicleMovementComponent.h
+++ b/Engine/Plugins/Experimental/ChaosVehiclesPlugin/Source/ChaosVehicles/Public/ChaosWheeledVehicleMovementComponent.h
@@ -86,6 +86,10 @@ struct CHAOSVEHICLES_API FWheelStatus
 	UPROPERTY()
 	FVector ContactPoint;
 
+	/** Wheel contact location */
+	UPROPERTY()
+	FVector HitLocation;
+
 	/** Material that wheel is in contact with */
 	UPROPERTY()
 	TWeakObjectPtr<class UPhysicalMaterial> PhysMaterial;
@@ -162,6 +166,7 @@ struct CHAOSVEHICLES_API FWheelStatus
 		SpringForce = 0.f;
 		SkidNormal = FVector::ZeroVector;
 		ContactPoint = FVector::ZeroVector;
+		HitLocation = FVector::ZeroVector;
 		bIsValid = false;
 		bABSActivated = false;
 		DriveTorque = 0.0f;
