Change the default spherecast method for vehicle suspension as the batched version is hitting an issue related to spheres against heightfields which doesn't present itself when using the non batch method. Can change back when other issue is found.

diff --git a/Engine/Plugins/Experimental/ChaosVehiclesPlugin/Source/ChaosVehicles/Public/ChaosVehicleMovementComponent.h b/Engine/Plugins/Experimental/ChaosVehiclesPlugin/Source/ChaosVehicles/Public/ChaosVehicleMovementComponent.h
--- a/Engine/Plugins/Experimental/ChaosVehiclesPlugin/Source/ChaosVehicles/Public/ChaosVehicleMovementComponent.h
+++ b/Engine/Plugins/Experimental/ChaosVehiclesPlugin/Source/ChaosVehicles/Public/ChaosVehicleMovementComponent.h
@@ -91,7 +91,7 @@ struct CHAOSVEHICLES_API FVehicleDebugParams
 	bool DisableAerodynamics = false;
 	bool DisableAerofoils = false;
 	bool DisableThrusters = false;
-	bool BatchQueries = true;
+	bool BatchQueries = false;
 	bool CacheTraceOverlap = false;
 	float ForceDebugScaling = 0.0006f;
 	float SleepCounterThreshold = 15;
