Network physics for vehicles

diff --git a/Engine/Plugins/Experimental/ChaosVehiclesPlugin/Source/ChaosVehicles/Private/ChaosVehicleMovementComponent.cpp b/Engine/Plugins/Experimental/ChaosVehiclesPlugin/Source/ChaosVehicles/Private/ChaosVehicleMovementComponent.cpp
--- a/Engine/Plugins/Experimental/ChaosVehiclesPlugin/Source/ChaosVehicles/Private/ChaosVehicleMovementComponent.cpp
+++ b/Engine/Plugins/Experimental/ChaosVehiclesPlugin/Source/ChaosVehicles/Private/ChaosVehicleMovementComponent.cpp
@@ -31,10 +31,14 @@
 #include "AI/Navigation/AvoidanceManager.h"
 #include "PhysicalMaterials/PhysicalMaterial.h"
 #include "GameFramework/HUD.h"
+#include "PhysicsEngine/PhysicsSettings.h"
 
+#include "PhysicsReplication.h"
+#include "Physics/Experimental/PhysScene_Chaos.h"
 #include "Chaos/Particle/ParticleUtilities.h"
 #include "Chaos/ParticleHandleFwd.h"
 #include "PhysicsProxy/SingleParticlePhysicsProxy.h"
+#include "PBDRigidsSolver.h"
 
 
 #define LOCTEXT_NAMESPACE "UVehicleMovementComponent"
@@ -68,6 +72,7 @@ FAutoConsoleVariableRef CVarChaosVehiclesSetMaxMPH(TEXT("p.Vehicle.SetMaxMPH"),
 FAutoConsoleVariableRef CVarChaosVehiclesEnableMultithreading(TEXT("p.Vehicle.EnableMultithreading"), GVehicleDebugParams.EnableMultithreading, TEXT("Enable multi-threading of vehicle updates."));
 FAutoConsoleVariableRef CVarChaosVehiclesControlInputWakeTolerance(TEXT("p.Vehicle.ControlInputWakeTolerance"), GVehicleDebugParams.ControlInputWakeTolerance, TEXT("Set the control input wake tolerance."));
 
+using namespace Chaos;
 
 void FVehicleState::CaptureState(const FBodyInstance* TargetInstance, float GravityZ, float DeltaTime)
 {
@@ -112,10 +117,20 @@ void FVehicleState::CaptureState(const Chaos::FRigidBodyHandle_Internal* Handle,
 		VehicleLocalVelocity = VehicleWorldTransform.InverseTransformVector(VehicleWorldVelocity);
 		LocalAcceleration = (VehicleLocalVelocity - LastFrameVehicleLocalVelocity) / DeltaTime;
 		LocalGForce = LocalAcceleration / FMath::Abs(GravityZ);
+
+#if !(UE_BUILD_SHIPPING || UE_BUILD_TEST)
+		if (Chaos::FPhysicsSolverBase::IsNetworkPhysicsPredictionEnabled() && Chaos::FPhysicsSolverBase::CanDebugNetworkPhysicsPrediction())
+		{
+			UE_LOG(LogVehicle, Log, TEXT("Vehicle state last velocity = %s | World Transform = %s | World Velocity = %s"),
+				*LastFrameVehicleLocalVelocity.ToString(), *VehicleWorldTransform.ToString(), *VehicleWorldVelocity.ToString());
+		}
+#endif
 		LastFrameVehicleLocalVelocity = VehicleLocalVelocity;
 
 		ForwardSpeed = FVector::DotProduct(VehicleWorldVelocity, VehicleForwardAxis);
 		ForwardsAcceleration = LocalAcceleration.X;
+
+
 	}
 }
 
@@ -131,10 +146,22 @@ void UChaosVehicleSimulation::TickVehicle(UWorld* WorldIn, float DeltaTime, cons
 	// movement updates and replication
 	if (World && RigidHandle)
 	{
+		if (CanSimulate() && Handle)
+		{
+			// Sync the inputs datas from the history
+			SyncHistoryInputs(InputData, Handle);
+		}
 		if (!VehicleState.bSleeping)
 		{
 			if (CanSimulate() && Handle)
 			{
+				// Update the vehicle/wheels... states
+				UpdateState(DeltaTime, InputData, Handle);
+
+				// Apply the controls inputs
+				ApplyInput(InputData.ControlInputs, DeltaTime);
+
+				// Update the simulation forces/impulses...
 				UpdateSimulation(DeltaTime, InputData, Handle);
 				FillOutputState(OutputData);
 			}
@@ -147,15 +174,323 @@ void UChaosVehicleSimulation::TickVehicle(UWorld* WorldIn, float DeltaTime, cons
 #endif
 }
 
+FVehicleStatesHistory::FVehicleStatesHistory(const int32 FrameCount, const bool bIsHistoryLocal, class UChaosVehicleSimulation* ChaosSimulation)
+	: Super(FrameCount, bIsHistoryLocal), VehicleSimulation(ChaosSimulation)
+{}
+
+FVehicleStatesHistory::~FVehicleStatesHistory()
+{}
+
+bool FVehicleStatesHistory::RecordStates(const int32 RecordFrame, const DatasType& ReplicatedDatas)
+{
+	if (VehicleSimulation && Super::RecordDatas(RecordFrame))
+	{
+		FVehicleStatesDatas& CurrentDatas = GetCurrentDatas();
+		CurrentDatas.StateLastVelocity = ReplicatedDatas.StateLastVelocity;
+		CurrentDatas.EngineOmega = ReplicatedDatas.EngineOmega;
+		CurrentDatas.SuspensionAveragedCount = ReplicatedDatas.SuspensionAveragedCount;
+		CurrentDatas.SuspensionAveragedNum = ReplicatedDatas.SuspensionAveragedNum;
+		CurrentDatas.SuspensionAveragedLength = ReplicatedDatas.SuspensionAveragedLength;
+		CurrentDatas.SuspensionLastDisplacement = ReplicatedDatas.SuspensionLastDisplacement;
+		CurrentDatas.SuspensionLastSpringLength = ReplicatedDatas.SuspensionLastSpringLength;
+		CurrentDatas.WheelsOmega = ReplicatedDatas.WheelsOmega;
+		CurrentDatas.WheelsAngularPosition = ReplicatedDatas.WheelsAngularPosition;
+
+#if !(UE_BUILD_SHIPPING || UE_BUILD_TEST)
+		if (Chaos::FPhysicsSolverBase::CanDebugNetworkPhysicsPrediction())
+		{
+			UE_LOG(LogVehicle, Log, TEXT("Record States at Frame %d on Vehicle = %d with Current = %d | Velocity = %s | Omega = %f"),
+				RecordFrame, VehicleSimulation, CurrentFrame, *CurrentDatas.StateLastVelocity.ToString(), CurrentDatas.EngineOmega);
+		}
+#endif
+		return true;
+	}
+	return false;
+}
+
+void FVehicleStatesHistory::ReplicateStates(DatasType& ReplicatedDatas)
+{
+	if (VehicleSimulation)
+	{
+		if (TUniquePtr<Chaos::FSimpleWheeledVehicle>& Vehicle = VehicleSimulation->PVehicle)
+		{
+			ReplicatedDatas.EngineOmega = Vehicle->GetEngine().GetEngineOmega();
+			ReplicatedDatas.StateLastVelocity = VehicleSimulation->VehicleState.LastFrameVehicleLocalVelocity;
+
+			const int32 NumWheels = Vehicle->Wheels.Num();
+
+			int32 NumLength = 0;
+			for (int32 WheelIdx = 0; WheelIdx < NumWheels; ++WheelIdx)
+			{
+				NumLength += Vehicle->GetSuspension(WheelIdx).GetAveragingNum();
+			}
+
+			ReplicatedDatas.SuspensionLastSpringLength.SetNum(NumWheels);
+			ReplicatedDatas.SuspensionLastDisplacement.SetNum(NumWheels);
+			ReplicatedDatas.SuspensionAveragedCount.SetNum(NumWheels);
+			ReplicatedDatas.SuspensionAveragedNum.SetNum(NumWheels);
+			ReplicatedDatas.SuspensionAveragedLength.SetNum(NumLength);
+			ReplicatedDatas.WheelsAngularPosition.SetNum(NumWheels);
+			ReplicatedDatas.WheelsOmega.SetNum(NumWheels);
+
+			int32 LengthCount = 0;
+			for (int32 WheelIdx = 0; WheelIdx < Vehicle->Wheels.Num(); ++WheelIdx)
+			{
+				FSimpleSuspensionSim& Suspension = Vehicle->GetSuspension(WheelIdx);
+				ReplicatedDatas.SuspensionLastSpringLength[WheelIdx] = Suspension.GetLastSpringLength();
+				ReplicatedDatas.SuspensionLastDisplacement[WheelIdx] = Suspension.GetLastDisplacement();
+				ReplicatedDatas.SuspensionAveragedCount[WheelIdx] = Suspension.GetAveragingCount();
+				ReplicatedDatas.SuspensionAveragedNum[WheelIdx] = Suspension.GetAveragingNum();
+
+				for (int32 LengthIdx = 0; LengthIdx < Suspension.GetAveragingNum(); ++LengthIdx)
+				{
+					ReplicatedDatas.SuspensionAveragedLength[LengthCount++] = Suspension.GetAveragingLength(LengthIdx);
+				}
+
+				FSimpleWheelSim& Wheel = Vehicle->GetWheel(WheelIdx);
+				ReplicatedDatas.WheelsOmega[WheelIdx] = Wheel.Omega;
+				ReplicatedDatas.WheelsAngularPosition[WheelIdx] = Wheel.AngularPosition;
+			}
+		}
+	}
+}
+
+bool FVehicleStatesHistory::RewindStates(const int32 RewindFrame, const bool bResetSolver)
+{
+	if (VehicleSimulation && Super::RewindStates(RewindFrame, bResetSolver))
+	{
+		const FVehicleStatesDatas& CurrentDatas = GetCurrentDatas();
+		VehicleSimulation->VehicleState.LastFrameVehicleLocalVelocity = CurrentDatas.StateLastVelocity;
+
+		if (TUniquePtr<Chaos::FSimpleWheeledVehicle>& Vehicle = VehicleSimulation->PVehicle)
+		{
+			Vehicle->GetEngine().SetEngineOmega(CurrentDatas.EngineOmega);
+
+			int32 LengthCount = 0;
+			for (int32 WheelIdx = 0, NumWheels = Vehicle->Wheels.Num(); WheelIdx < NumWheels; ++WheelIdx)
+			{
+				FSimpleSuspensionSim& Suspension = Vehicle->GetSuspension(WheelIdx);
+				Suspension.SetLastSpringLength(CurrentDatas.SuspensionLastSpringLength[WheelIdx]);
+				Suspension.SetLastDisplacement(CurrentDatas.SuspensionLastDisplacement[WheelIdx]);
+				Suspension.SetAveragingCount(CurrentDatas.SuspensionAveragedCount[WheelIdx]);
+				Suspension.SetAveragingNum(CurrentDatas.SuspensionAveragedNum[WheelIdx]);
+
+				for (int32 LengthIdx = 0; LengthIdx < Suspension.GetAveragingNum(); ++LengthIdx)
+				{
+					Suspension.SetAveragingLength(LengthIdx, CurrentDatas.SuspensionAveragedLength[LengthCount++]);
+				}
+
+				FSimpleWheelSim& Wheel = Vehicle->GetWheel(WheelIdx);
+				Wheel.Omega = CurrentDatas.WheelsOmega[WheelIdx];
+				Wheel.AngularPosition = CurrentDatas.WheelsAngularPosition[WheelIdx];
+			}
+#if !(UE_BUILD_SHIPPING || UE_BUILD_TEST)
+			if (Chaos::FPhysicsSolverBase::CanDebugNetworkPhysicsPrediction())
+			{
+				UE_LOG(LogVehicle, Log, TEXT("Rewind States at Frame %d on Vehicle = %d with Current = %d | Velocity = %s | Omega = %f"),
+					RewindFrame, VehicleSimulation, CurrentFrame, *CurrentDatas.StateLastVelocity.ToString(), CurrentDatas.EngineOmega);
+			}
+#endif
+		}
+		return true;
+	}
+	return false;
+}
+
+FVehicleInputsHistory::FVehicleInputsHistory(const int32 FrameCount, const bool bIsHistoryLocal, class UChaosVehicleSimulation* ChaosSimulation)
+	: Super(FrameCount, bIsHistoryLocal), VehicleSimulation(ChaosSimulation)
+{}
+
+FVehicleInputsHistory::~FVehicleInputsHistory()
+{}
+
+bool FVehicleInputsHistory::RecordInputs(const int32 RecordFrame, const DatasType& ReplicatedDatas)
+{
+	if (VehicleSimulation && Super::RecordDatas(RecordFrame))
+	{
+		FVehicleInputsDatas& CurrentDatas = GetCurrentDatas();
+		CurrentDatas.VehicleInputs.SteeringInput = ReplicatedDatas.VehicleInputs.SteeringInput;
+		CurrentDatas.VehicleInputs.ThrottleInput = ReplicatedDatas.VehicleInputs.ThrottleInput;
+		CurrentDatas.VehicleInputs.BrakeInput = ReplicatedDatas.VehicleInputs.BrakeInput;
+		CurrentDatas.VehicleInputs.HandbrakeInput = ReplicatedDatas.VehicleInputs.HandbrakeInput;
+		CurrentDatas.VehicleInputs.RollInput = ReplicatedDatas.VehicleInputs.RollInput;
+		CurrentDatas.VehicleInputs.PitchInput = ReplicatedDatas.VehicleInputs.PitchInput;
+		CurrentDatas.VehicleInputs.YawInput = ReplicatedDatas.VehicleInputs.YawInput;
+		CurrentDatas.TransmissionCurrentGear = ReplicatedDatas.TransmissionCurrentGear;
+		CurrentDatas.TransmissionTargetGear = ReplicatedDatas.TransmissionTargetGear;
+		CurrentDatas.TransmissionChangeTime = ReplicatedDatas.TransmissionChangeTime;
+
+#if !(UE_BUILD_SHIPPING || UE_BUILD_TEST)
+		if (Chaos::FPhysicsSolverBase::CanDebugNetworkPhysicsPrediction())
+		{
+			UE_LOG(LogVehicle, Log, TEXT("Record Inputs at Frame %d on Vehicle = %d with Current = %d | Brake = %f | Handbrake = %f | Pitch = %f | Roll = %f | Steering = %f | Throttle = %f | Yaw = %f | Gear = %d"),
+				RecordFrame, VehicleSimulation, CurrentFrame, CurrentDatas.VehicleInputs.BrakeInput, CurrentDatas.VehicleInputs.HandbrakeInput,
+				CurrentDatas.VehicleInputs.PitchInput, CurrentDatas.VehicleInputs.RollInput, CurrentDatas.VehicleInputs.SteeringInput,
+				CurrentDatas.VehicleInputs.ThrottleInput, CurrentDatas.VehicleInputs.YawInput, CurrentDatas.TransmissionCurrentGear);
+		}
+#endif
+		return true;
+	}
+	return false;
+}
+
+void FVehicleInputsHistory::ReplicateInputs(
+	const float SteeringInput, const float ThrottleInput,
+	const float BrakeInput, const float HandbrakeInput, const float RollInput,
+	const float PitchInput, const float YawInput, const int32 CurrentGear, DatasType& ReplicatedDatas)
+{
+	ReplicatedDatas.VehicleInputs.SteeringInput = SteeringInput;
+	ReplicatedDatas.VehicleInputs.ThrottleInput = ThrottleInput;
+	ReplicatedDatas.VehicleInputs.BrakeInput = BrakeInput;
+	ReplicatedDatas.VehicleInputs.HandbrakeInput = HandbrakeInput;
+	ReplicatedDatas.VehicleInputs.RollInput = RollInput;
+	ReplicatedDatas.VehicleInputs.PitchInput = PitchInput;
+	ReplicatedDatas.VehicleInputs.YawInput = YawInput;
+
+	if (VehicleSimulation)
+	{
+		if (TUniquePtr<Chaos::FSimpleWheeledVehicle>& Vehicle = VehicleSimulation->PVehicle)
+		{
+			if (Vehicle->HasTransmission())
+			{
+				FSimpleTransmissionSim& Transmission = Vehicle->GetTransmission();
+
+				Transmission.SetGear(CurrentGear, true);
+				ReplicatedDatas.TransmissionCurrentGear = Transmission.GetCurrentGear();
+				ReplicatedDatas.TransmissionTargetGear = Transmission.GetCurrentGear();
+				ReplicatedDatas.TransmissionChangeTime = Transmission.GetCurrentGearChangeTime();
+			}
+		}
+	}
+}
+
+void FVehicleInputsHistory::SmoothInputs(const int32 ApplyFrame, const int32 ValidFrame, const bool bApplySmoothing)
+{
+	if (!bIsLocalHistory && VehicleSimulation)
+	{
+		const float HalfFrames = (float)NumFrames / 2.0f;
+		const float GaussSigma = HalfFrames / 2.0f;
+
+		const int32 ValidIndex = ValidFrame % NumFrames;
+		const int32 FrameIndex = ValidFrame - ApplyFrame;
+
+		const float SmoothFactor = bApplySmoothing ? FMath::Exp(-FrameIndex * FrameIndex / (2.0f * GaussSigma * GaussSigma)) : 1.0;
+
+		VehicleSimulation->VehicleInputs.BrakeInput = SmoothFactor * DatasArray[ValidIndex].VehicleInputs.BrakeInput;
+		VehicleSimulation->VehicleInputs.HandbrakeInput = SmoothFactor * DatasArray[ValidIndex].VehicleInputs.HandbrakeInput;
+		VehicleSimulation->VehicleInputs.PitchInput = SmoothFactor * DatasArray[ValidIndex].VehicleInputs.PitchInput;
+		VehicleSimulation->VehicleInputs.RollInput = SmoothFactor * DatasArray[ValidIndex].VehicleInputs.RollInput;
+		VehicleSimulation->VehicleInputs.ThrottleInput = SmoothFactor * DatasArray[ValidIndex].VehicleInputs.ThrottleInput;
+		VehicleSimulation->VehicleInputs.SteeringInput = SmoothFactor * DatasArray[ValidIndex].VehicleInputs.SteeringInput;
+		VehicleSimulation->VehicleInputs.YawInput = SmoothFactor * DatasArray[ValidIndex].VehicleInputs.YawInput;
+
+		if (TUniquePtr<Chaos::FSimpleWheeledVehicle>& Vehicle = VehicleSimulation->PVehicle)
+		{
+			if (Vehicle->HasTransmission())
+			{
+				FSimpleTransmissionSim& Transmission = VehicleSimulation->PVehicle->GetTransmission();
+				Transmission.SetCurrentGear(DatasArray[ValidIndex].TransmissionCurrentGear);
+				Transmission.SetTargetGear(DatasArray[ValidIndex].TransmissionTargetGear);
+				Transmission.SetCurrentGearChangeTime(DatasArray[ValidIndex].TransmissionChangeTime);
+			}
+
+#if !(UE_BUILD_SHIPPING || UE_BUILD_TEST)
+			if (Chaos::FPhysicsSolverBase::CanDebugNetworkPhysicsPrediction())
+			{
+				UE_LOG(LogVehicle, Log, TEXT("Smooth Inputs from frame %d to frame = %d with factor %f"),
+					ValidFrame, ApplyFrame, SmoothFactor);
+			}
+#endif
+		}
+	}
+}
+
+bool FVehicleInputsHistory::ApplyInputs(const int32 ApplyFrame, const bool bResetSolver)
+{
+	if (VehicleSimulation )
+	{
+		if(Super::ApplyInputs(ApplyFrame, bResetSolver))
+		{
+			const FVehicleInputsDatas& CurrentDatas = GetCurrentDatas();
+
+			VehicleSimulation->VehicleInputs = CurrentDatas.VehicleInputs;
+
+			FSimpleTransmissionSim& Transmission = VehicleSimulation->PVehicle->GetTransmission();
+			Transmission.SetCurrentGear(CurrentDatas.TransmissionCurrentGear);
+			Transmission.SetTargetGear(CurrentDatas.TransmissionTargetGear);
+			Transmission.SetCurrentGearChangeTime(CurrentDatas.TransmissionChangeTime);
+		}
+		else
+		{
+			int32 ClosestFrame = INDEX_NONE;
+			for (int32 FrameIndex = 0; FrameIndex < NumFrames; ++FrameIndex)
+			{
+				const int32 ValidFrame = FMath::Max(0, ApplyFrame - FrameIndex);
+				const int32 ValidIndex = ValidFrame % NumFrames;
+
+				if (DatasArray[ValidIndex].DatasFrame == ValidFrame)
+				{
+					ClosestFrame = ValidFrame;
+					break;
+				}
+			}
+			if(ClosestFrame != INDEX_NONE)
+			{ 
+				SmoothInputs(ApplyFrame, ClosestFrame, false);
+			}
+		}
+
+#if !(UE_BUILD_SHIPPING || UE_BUILD_TEST)
+		FSimpleTransmissionSim& Transmission = VehicleSimulation->PVehicle->GetTransmission();
+
+		if (Chaos::FPhysicsSolverBase::CanDebugNetworkPhysicsPrediction())
+		{
+			UE_LOG(LogVehicle, Log, TEXT("Apply Inputs at Frame %d on Vehicle = %d with Current = %d | Brake = %f | Handbrake = %f | Pitch = %f | Roll = %f | Steering = %f | Throttle = %f | Yaw = %f | Gear = %d"),
+				ApplyFrame, VehicleSimulation, CurrentFrame, VehicleSimulation->VehicleInputs.BrakeInput,
+				VehicleSimulation->VehicleInputs.HandbrakeInput, VehicleSimulation->VehicleInputs.PitchInput,
+				VehicleSimulation->VehicleInputs.RollInput, VehicleSimulation->VehicleInputs.SteeringInput,
+				VehicleSimulation->VehicleInputs.ThrottleInput, VehicleSimulation->VehicleInputs.YawInput,
+				Transmission.GetCurrentGear());
+
+		}
+#endif
+		return true;
+	}
+	return false;
+}
+
+void UChaosVehicleSimulation::SyncHistoryInputs(const FChaosVehicleDefaultAsyncInput& InputData, Chaos::FRigidBodyHandle_Internal* Handle)
+{
+	if (Chaos::FPhysicsSolverBase::IsNetworkPhysicsPredictionEnabled())
+	{
+		InputData.ControlInputs.CopyVehicleInputs(VehicleInputs);
+
+#if !(UE_BUILD_SHIPPING || UE_BUILD_TEST)
+		if (Chaos::FPhysicsSolverBase::CanDebugNetworkPhysicsPrediction())
+		{
+			auto* RigidSolver = Handle->GetProxy()->GetSolver<Chaos::FPBDRigidsSolver>();
+			if (RigidSolver && !RigidSolver->GetEvolution()->IsResimming())
+			{
+				UE_LOG(LogVehicle, Log, TEXT("Store States at Frame %d on Vehicle = %d | Velocity = %s | Omega = %f"),
+					RigidSolver->GetCurrentFrame(), this, *VehicleState.LastFrameVehicleLocalVelocity.ToString(), PVehicle->GetEngine().GetEngineOmega());
+			}
+		}
+	}
+#endif
+}
+
 void UChaosVehicleSimulation::ApplyDeferredForces(Chaos::FRigidBodyHandle_Internal* Handle)
 {
 	DeferredForces.Apply(Handle);
 }
 
-void UChaosVehicleSimulation::UpdateSimulation(float DeltaTime, const FChaosVehicleDefaultAsyncInput& InputData, Chaos::FRigidBodyHandle_Internal* Handle)
+void UChaosVehicleSimulation::UpdateState(float DeltaTime, const FChaosVehicleDefaultAsyncInput& InputData, Chaos::FRigidBodyHandle_Internal* Handle)
 {
 	VehicleState.CaptureState(Handle, InputData.GravityZ, DeltaTime);
+}
 
+void UChaosVehicleSimulation::UpdateSimulation(float DeltaTime, const FChaosVehicleDefaultAsyncInput& InputData, Chaos::FRigidBodyHandle_Internal* Handle)
+{
 	ApplyAerodynamics(DeltaTime);
 	ApplyAerofoilForces(DeltaTime);
 	ApplyThrustForces(DeltaTime);
@@ -163,9 +498,7 @@ void UChaosVehicleSimulation::UpdateSimulation(float DeltaTime, const FChaosVehi
 }
 
 void UChaosVehicleSimulation::FillOutputState(FChaosVehicleAsyncOutput& Output)
-{
-
-}
+{}
 
 /** Pass control Input to the vehicle systems */
 void UChaosVehicleSimulation::ApplyInput(const FControlInputs& ControlInputs, float DeltaTime)
@@ -242,7 +575,6 @@ void UChaosVehicleSimulation::ApplyInput(const FControlInputs& ControlInputs, fl
 
 }
 
-
 void UChaosVehicleSimulation::ApplyAerodynamics(float DeltaTime)
 {
 	if (!GVehicleDebugParams.DisableAerodynamics)
@@ -451,7 +783,7 @@ void UChaosVehicleSimulation::ApplyTorqueControl(float DeltaTime, const FChaosVe
 		}
 		AddForce(CorrectionalForce);
 	}
-	
+
 }
 
 void UChaosVehicleSimulation::DrawDebug3D()
@@ -555,7 +887,6 @@ void UChaosVehicleSimulation::InitializeSuspension(int WheelIndex, const Chaos::
 	}
 }
 
-
 /**
  * UChaosVehicleMovementComponent
  */
@@ -610,6 +941,7 @@ UChaosVehicleMovementComponent::UChaosVehicleMovementComponent(const FObjectInit
 	TransmissionType = Chaos::ETransmissionType::Automatic;
 
 	SetIsReplicatedByDefault(true);
+	bUsingNetworkPhysicsPrediction = Chaos::FPhysicsSolverBase::IsNetworkPhysicsPredictionEnabled();
 
 	AHUD::OnShowDebugInfo.AddUObject(this, &UChaosVehicleMovementComponent::ShowDebugInfo);
 }
@@ -712,17 +1044,41 @@ void UChaosVehicleMovementComponent::OnCreatePhysicsState()
 	UWorld* World = GetWorld();
 	if (World->IsGameWorld())
 	{
-		FPhysScene* PhysScene = World->GetPhysicsScene();
-
-		if (PhysScene && FChaosVehicleManager::GetVehicleManagerFromScene(PhysScene))
+		if (FPhysScene* PhysScene = World->GetPhysicsScene())
 		{
-			CreateVehicle();
-			FixupSkeletalMesh();
+			if (FChaosVehicleManager::GetVehicleManagerFromScene(PhysScene))
+			{
+				CreateVehicle();
+				FixupSkeletalMesh();
 
-			if (PVehicleOutput)
+				if (PVehicleOutput)
+				{
+					FChaosVehicleManager* VehicleManager = FChaosVehicleManager::GetVehicleManagerFromScene(PhysScene);
+					VehicleManager->AddVehicle(this);
+				}
+			}
+			if (bUsingNetworkPhysicsPrediction)
 			{
-				FChaosVehicleManager* VehicleManager = FChaosVehicleManager::GetVehicleManagerFromScene(PhysScene);
-				VehicleManager->AddVehicle(this);
+				if (Chaos::FPhysicsSolver* Solver = PhysScene->GetSolver())
+				{
+					if (Chaos::FRewindData* RewindData = Solver->GetRewindData())
+					{
+						if (VehicleSimulationPT.IsValid())
+						{
+							AController* Controller = GetController();
+							const bool bIsLocalHistory = (Controller && Controller->IsLocalController());
+
+							const IConsoleVariable* NumFramesCVar = IConsoleManager::Get().FindConsoleVariable(TEXT("p.RewindCaptureNumFrames"));
+							const int32 NumFrames = NumFramesCVar ? NumFramesCVar->GetInt() : 50;
+
+							VehicleSimulationPT->InputsHistory = MakeShared<FVehicleInputsHistory>(NumFrames, bIsLocalHistory, VehicleSimulationPT.Get());
+							RewindData->AddInputsHistory(VehicleSimulationPT->InputsHistory);
+
+							VehicleSimulationPT->StatesHistory = MakeShared<FVehicleStatesHistory>(NumFrames, bIsLocalHistory, VehicleSimulationPT.Get());
+							RewindData->AddStatesHistory(VehicleSimulationPT->StatesHistory);
+						}
+					}
+				}
 			}
 		}
 	}
@@ -750,6 +1106,27 @@ void UChaosVehicleMovementComponent::OnDestroyPhysicsState()
 			UpdatedComponent->RecreatePhysicsState();
 		}
 	}
+	if (bUsingNetworkPhysicsPrediction)
+	{
+		UWorld* World = GetWorld();
+		if (World->IsGameWorld())
+		{
+			if (FPhysScene* PhysScene = World->GetPhysicsScene())
+			{
+				if (Chaos::FPhysicsSolver* Solver = PhysScene->GetSolver())
+				{
+					if (Chaos::FRewindData* RewindData = Solver->GetRewindData())
+					{
+						if (VehicleSimulationPT.IsValid())
+						{
+							RewindData->RemoveInputsHistory(VehicleSimulationPT->InputsHistory);
+							RewindData->RemoveStatesHistory(VehicleSimulationPT->StatesHistory);
+						}
+					}
+				}
+			}
+		}
+	}
 }
 
 void UChaosVehicleMovementComponent::PreTickGT(float DeltaTime)
@@ -764,8 +1141,8 @@ void UChaosVehicleMovementComponent::PreTickGT(float DeltaTime)
 		}
 	}
 
+	if (!bUsingNetworkPhysicsPrediction)
 	{
-		// is this needless copying
 		FControlInputs ControlInputs;
 		ControlInputs.ThrottleInput = ThrottleInput;
 		ControlInputs.BrakeInput = BrakeInput;
@@ -778,7 +1155,6 @@ void UChaosVehicleMovementComponent::PreTickGT(float DeltaTime)
 		ProcessSleeping(ControlInputs);
 	}
 
-
 	if (VehicleSetupTag != FChaosVehicleManager::VehicleSetupTag)
 	{
 		RecreatePhysicsState();
@@ -787,6 +1163,11 @@ void UChaosVehicleMovementComponent::PreTickGT(float DeltaTime)
 
 void UChaosVehicleMovementComponent::StopMovementImmediately()
 {
+	if (bUsingNetworkPhysicsPrediction)
+	{
+		return; // Called when player onwership changes I think - client restarts - resets velocity/forces bad for our case.
+	}
+
 	FBodyInstance* TargetInstance = GetBodyInstance();
 	if (TargetInstance)
 	{
@@ -884,17 +1265,20 @@ void UChaosVehicleMovementComponent::SetTargetGear(int32 GearNum, bool bImmediat
 {
 	if (PVehicleOutput && GearNum != PVehicleOutput->TargetGear)
 	{
-		FBodyInstance* TargetInstance = UpdatedPrimitive->GetBodyInstance();
-
-		if (TargetInstance)
+		if (!bUsingNetworkPhysicsPrediction)
 		{
-			FPhysicsCommand::ExecuteWrite(TargetInstance->ActorHandle, [&](const FPhysicsActorHandle& Chassis)
-				{
-					if (VehicleSimulationPT && VehicleSimulationPT->PVehicle && VehicleSimulationPT->PVehicle->HasTransmission())
+			FBodyInstance* TargetInstance = UpdatedPrimitive->GetBodyInstance();
+
+			if (TargetInstance)
+			{
+				FPhysicsCommand::ExecuteWrite(TargetInstance->ActorHandle, [&](const FPhysicsActorHandle& Chassis)
 					{
-						VehicleSimulationPT->PVehicle->GetTransmission().SetGear(GearNum, bImmediate);
-					}
-				});
+						if (VehicleSimulationPT && VehicleSimulationPT->PVehicle && VehicleSimulationPT->PVehicle->HasTransmission())
+						{
+							VehicleSimulationPT->PVehicle->GetTransmission().SetGear(GearNum, bImmediate);
+						}
+					});
+			}
 		}
 
 		TargetGear = GearNum;
@@ -915,7 +1299,7 @@ void UChaosVehicleMovementComponent::SetRequiresControllerForInputs(bool bRequir
 
 int32 UChaosVehicleMovementComponent::GetCurrentGear() const
 {
-	return (PVehicleOutput)?PVehicleOutput->CurrentGear:0;
+	return (PVehicleOutput) ? PVehicleOutput->CurrentGear : 0;
 }
 
 int32 UChaosVehicleMovementComponent::GetTargetGear() const
@@ -1060,7 +1444,10 @@ void UChaosVehicleMovementComponent::ClearInput()
 	AController* Controller = GetController();
 	if (Controller && Controller->IsLocalController() && PVehicleOutput)
 	{
-		ServerUpdateState(SteeringInput, ThrottleInput, BrakeInput, HandbrakeInput, CurrentGear, RollInput, PitchInput, YawInput);
+		if (!bUsingNetworkPhysicsPrediction)
+		{
+			ServerUpdateState(SteeringInput, ThrottleInput, BrakeInput, HandbrakeInput, CurrentGear, RollInput, PitchInput, YawInput);
+		}
 	}
 }
 
@@ -1104,7 +1491,7 @@ void UChaosVehicleMovementComponent::UpdateState(float DeltaTime)
 	}
 	VehicleState.bAllWheelsOnGround = (VehicleState.NumWheelsOnGround == NumWheels);
 
-	bool bProcessLocally = bRequiresControllerForInputs?(Controller && Controller->IsLocalController()):true;
+	bool bProcessLocally = bRequiresControllerForInputs ? (Controller && Controller->IsLocalController()) : true;
 
 	// IsLocallyControlled will fail if the owner is unpossessed (i.e. Controller == nullptr);
 	// Should we remove input instead of relying on replicated state in that case?
@@ -1143,6 +1530,8 @@ void UChaosVehicleMovementComponent::UpdateState(float DeltaTime)
 		float ModifiedThrottle = 0.f;
 		float ModifiedBrake = 0.f;
 		CalcThrottleBrakeInput(ModifiedThrottle, ModifiedBrake);
+
+		// Apply Inputs locally
 		SteeringInput = SteeringInputRate.InterpInputValue(DeltaTime, SteeringInput, CalcSteeringInput());
 		ThrottleInput = ThrottleInputRate.InterpInputValue(DeltaTime, ThrottleInput, ModifiedThrottle);
 		BrakeInput = BrakeInputRate.InterpInputValue(DeltaTime, BrakeInput, ModifiedBrake);
@@ -1151,15 +1540,37 @@ void UChaosVehicleMovementComponent::UpdateState(float DeltaTime)
 		YawInput = YawInputRate.InterpInputValue(DeltaTime, YawInput, CalcYawInput());
 		HandbrakeInput = HandbrakeInputRate.InterpInputValue(DeltaTime, HandbrakeInput, CalcHandbrakeInput());
 
-		// and send to server - (ServerUpdateState_Implementation below)
-		ServerUpdateState(SteeringInput, ThrottleInput, BrakeInput, HandbrakeInput, GetCurrentGear(), RollInput, PitchInput, YawInput);
+		if (!bUsingNetworkPhysicsPrediction)
+		{
+			// and send to server - (ServerUpdateState_Implementation below)
+			ServerUpdateState(SteeringInput, ThrottleInput, BrakeInput, HandbrakeInput, GetCurrentGear(), RollInput, PitchInput, YawInput);
+		}
+		else
+		{
+			if (APawn* MyPawn = Cast<APawn>(GetOwner()))
+			{
+				if (APlayerController* PlayerController = Cast<APlayerController>(MyPawn->GetController()))
+				{
+					if (PlayerController->IsLocalController() && !PlayerController->IsNetMode(NM_DedicatedServer))
+					{
+						const FAsyncPhysicsTimestamp TimeStamp = PlayerController->GetAsyncPhysicsTimestamp();
+						ApplyControls_Imp(SteeringInput, ThrottleInput, BrakeInput, HandbrakeInput, TargetGear, RollInput, PitchInput, YawInput, TimeStamp, true);
+
+						if (!PlayerController->IsNetMode(NM_ListenServer))
+						{
+							ServerApplyControls(SteeringInput, ThrottleInput, BrakeInput, HandbrakeInput, TargetGear, RollInput, PitchInput, YawInput, TimeStamp);
+						}
+					}
+				}
+			}
+		}
 
 		if (PawnOwner && PawnOwner->IsNetMode(NM_Client))
 		{
 			MarkForClientCameraUpdate();
 		}
 	}
-	else
+	else if (!bUsingNetworkPhysicsPrediction)
 	{
 		// use replicated values for remote pawns
 		SteeringInput = ReplicatedState.SteeringInput;
@@ -1191,16 +1602,24 @@ void UChaosVehicleMovementComponent::ProcessSleeping(const FControlInputs& Contr
 		// If the vehicle is locally controlled, we want to use the raw inputs to determine sleep.
 		// However, if it's on the Server or is just being replicated to other Clients then there
 		// won't be any Raw input. In that case, use ReplicatedState instead.
-		
+
 		// NOTE: Even on local clients, ReplicatedState will still be populated (the call to ServerUpdateState will
 		//			be processed locally). Maybe we should *just* use ReplicatedState?
 
 		const AController* Controller = GetController();
 		const bool bIsLocallyControlled = (Controller && Controller->IsLocalController());
-		const bool bControlInputPressed = bIsLocallyControlled ? (ControlInputs.ThrottleInput >= GVehicleDebugParams.ControlInputWakeTolerance) || (FMath::Abs(ControlInputs.SteeringInput - PrevSteeringInput) >= GVehicleDebugParams.ControlInputWakeTolerance)
-			|| (ControlInputs.RollInput >= GVehicleDebugParams.ControlInputWakeTolerance) || (ControlInputs.PitchInput >= GVehicleDebugParams.ControlInputWakeTolerance) || (ControlInputs.YawInput >= GVehicleDebugParams.ControlInputWakeTolerance)
-			: (ReplicatedState.ThrottleInput >= GVehicleDebugParams.ControlInputWakeTolerance) || (FMath::Abs(ReplicatedState.SteeringInput - PrevReplicatedSteeringInput) >= GVehicleDebugParams.ControlInputWakeTolerance)
-			|| (ReplicatedState.RollInput >= GVehicleDebugParams.ControlInputWakeTolerance) || (ReplicatedState.PitchInput >= GVehicleDebugParams.ControlInputWakeTolerance) || (ReplicatedState.YawInput >= GVehicleDebugParams.ControlInputWakeTolerance);
+		const bool bControlInputPressed = bIsLocallyControlled ? (ControlInputs.ThrottleInput >= GVehicleDebugParams.ControlInputWakeTolerance)
+			|| (ControlInputs.BrakeInput >= GVehicleDebugParams.ControlInputWakeTolerance)
+			|| (FMath::Abs(ControlInputs.SteeringInput - PrevSteeringInput) >= GVehicleDebugParams.ControlInputWakeTolerance)
+			|| (ControlInputs.RollInput >= GVehicleDebugParams.ControlInputWakeTolerance)
+			|| (ControlInputs.PitchInput >= GVehicleDebugParams.ControlInputWakeTolerance)
+			|| (ControlInputs.YawInput >= GVehicleDebugParams.ControlInputWakeTolerance)
+			: (ReplicatedState.ThrottleInput >= GVehicleDebugParams.ControlInputWakeTolerance)
+			|| (ReplicatedState.BrakeInput >= GVehicleDebugParams.ControlInputWakeTolerance)
+			|| (FMath::Abs(ReplicatedState.SteeringInput - PrevReplicatedSteeringInput) >= GVehicleDebugParams.ControlInputWakeTolerance)
+			|| (ReplicatedState.RollInput >= GVehicleDebugParams.ControlInputWakeTolerance)
+			|| (ReplicatedState.PitchInput >= GVehicleDebugParams.ControlInputWakeTolerance)
+			|| (ReplicatedState.YawInput >= GVehicleDebugParams.ControlInputWakeTolerance);
 
 		PrevSteeringInput = ControlInputs.SteeringInput;
 		PrevReplicatedSteeringInput = ReplicatedState.SteeringInput;
@@ -1215,7 +1634,7 @@ void UChaosVehicleMovementComponent::ProcessSleeping(const FControlInputs& Contr
 		else if (!GVehicleDebugParams.DisableVehicleSleep && !VehicleState.bSleeping && !bControlInputPressed && VehicleState.bAllWheelsOnGround && (VehicleState.VehicleUpAxis.Z > SleepSlopeLimit))
 		{
 			float SpeedSqr = TargetInstance->GetUnrealWorldVelocity().SizeSquared();
-			if (SpeedSqr < (SleepThreshold* SleepThreshold))
+			if (SpeedSqr < (SleepThreshold * SleepThreshold))
 			{
 				if (VehicleState.SleepCounter < GVehicleDebugParams.SleepCounterThreshold)
 				{
@@ -1272,6 +1691,57 @@ void UChaosVehicleMovementComponent::ServerUpdateState_Implementation(float InSt
 
 }
 
+void UChaosVehicleMovementComponent::ServerApplyControls_Implementation(float InSteeringInput, float InThrottleInput, float InBrakeInput, float InHandbrakeInput,
+	int32 InCurrentGear, float InRollInput, float InPitchInput, float InYawInput, const FAsyncPhysicsTimestamp& AsyncPhysicsTimestamp)
+{
+	ApplyControls_Imp(InSteeringInput, InThrottleInput, InBrakeInput, InHandbrakeInput, InCurrentGear, InRollInput, InPitchInput, InYawInput, AsyncPhysicsTimestamp, false);
+}
+
+void UChaosVehicleMovementComponent::ApplyControls_Imp(float InSteeringInput, float InThrottleInput, float InBrakeInput, float InHandbrakeInput,
+	int32 InCurrentGear, float InRollInput, float InPitchInput, float InYawInput, const FAsyncPhysicsTimestamp& AsyncPhysicsTimestamp, const bool bEnableResim)
+{
+	if (UWorld* World = GetWorld())
+	{
+		APawn* LocalPawn = Cast<APawn>(GetOwner());
+		if (APlayerController* PlayerController = (LocalPawn && LocalPawn->GetController()) ?
+			Cast<APlayerController>(LocalPawn->GetController()) : World->GetFirstPlayerController())
+		{
+			const bool bControlInputPressed = ((InThrottleInput >= GVehicleDebugParams.ControlInputWakeTolerance)
+				|| (InBrakeInput >= GVehicleDebugParams.ControlInputWakeTolerance)
+				|| (InRollInput >= GVehicleDebugParams.ControlInputWakeTolerance)
+				|| (InPitchInput >= GVehicleDebugParams.ControlInputWakeTolerance)
+				|| (InYawInput >= GVehicleDebugParams.ControlInputWakeTolerance)
+				|| (FMath::Abs(InSteeringInput - PrevSteeringInput) >= GVehicleDebugParams.ControlInputWakeTolerance));
+
+			if (bControlInputPressed)
+			{
+				SetSleeping(false);
+			}
+			const FAsyncPhysicsTimestamp LocalStep = PlayerController->GetAsyncPhysicsTimestamp();
+
+			PlayerController->ExecuteAsyncPhysicsCommand(AsyncPhysicsTimestamp, this, [this, InSteeringInput, InThrottleInput, InBrakeInput, InHandbrakeInput, InCurrentGear, InRollInput, InPitchInput, InYawInput, AsyncPhysicsTimestamp, LocalStep]
+			{
+				const bool bIsServer = (GetWorld()->IsNetMode(NM_ListenServer) || GetWorld()->IsNetMode(NM_DedicatedServer));
+				const int32 HistoryFrame = bIsServer ? AsyncPhysicsTimestamp.ServerFrame : AsyncPhysicsTimestamp.LocalFrame;
+
+				ReplicatedDatas.ServerFrame = AsyncPhysicsTimestamp.ServerFrame;
+
+				// Replicate the Inputs to send them to other clients
+				VehicleSimulationPT->InputsHistory->ReplicateInputs(InSteeringInput, InThrottleInput, 
+					InBrakeInput, InHandbrakeInput, InRollInput, InPitchInput, InYawInput, InCurrentGear, ReplicatedDatas.InputsDatas);
+
+				// We record in the history buffer the local client control inputs since they are the ground truth
+				VehicleSimulationPT->InputsHistory->RecordInputs(HistoryFrame, ReplicatedDatas.InputsDatas);
+					
+				// On the server we replicate the vehicle state in order to be used during resimulation at the rewind step
+				if (bIsServer)
+				{
+					VehicleSimulationPT->StatesHistory->ReplicateStates(ReplicatedDatas.StatesDatas);
+				}
+			}, bEnableResim);
+		}
+	}
+}
 /// @endcond
 
 
@@ -1509,12 +1979,13 @@ void UChaosVehicleMovementComponent::DrawDebug(UCanvas* Canvas, float& YL, float
 
 /// @cond DOXYGEN_WARNINGS
 
-void UChaosVehicleMovementComponent::GetLifetimeReplicatedProps( TArray< FLifetimeProperty > & OutLifetimeProps ) const
+void UChaosVehicleMovementComponent::GetLifetimeReplicatedProps(TArray< FLifetimeProperty >& OutLifetimeProps) const
 {
-	Super::GetLifetimeReplicatedProps( OutLifetimeProps );
+	Super::GetLifetimeReplicatedProps(OutLifetimeProps);
 
-	DOREPLIFETIME( UChaosVehicleMovementComponent, ReplicatedState );
+	DOREPLIFETIME(UChaosVehicleMovementComponent, ReplicatedState);
 	DOREPLIFETIME(UChaosVehicleMovementComponent, OverrideController);
+	DOREPLIFETIME(UChaosVehicleMovementComponent, ReplicatedDatas);
 }
 
 /// @endcond
@@ -1706,25 +2177,28 @@ void UChaosVehicleMovementComponent::Update(float DeltaTime)
 			if (auto Handle = BodyInstance->ActorHandle)
 			{
 				CurAsyncInput->Proxy = Handle;	// vehicles are never static
-
 				FChaosVehicleDefaultAsyncInput* AsyncInput = static_cast<FChaosVehicleDefaultAsyncInput*>(CurAsyncInput);
-				AsyncInput->ControlInputs.ThrottleInput = ThrottleInputRate.CalcControlFunction(ThrottleInput);
-				AsyncInput->ControlInputs.BrakeInput = BrakeInputRate.CalcControlFunction(BrakeInput);
-				AsyncInput->ControlInputs.SteeringInput = SteeringInputRate.CalcControlFunction(SteeringInput);
-				AsyncInput->ControlInputs.HandbrakeInput = HandbrakeInput;
-				AsyncInput->ControlInputs.RollInput = RollInputRate.CalcControlFunction(RollInput);
-				AsyncInput->ControlInputs.PitchInput = PitchInputRate.CalcControlFunction(PitchInput);
-				AsyncInput->ControlInputs.YawInput = YawInputRate.CalcControlFunction(YawInput);
-				AsyncInput->ControlInputs.GearUpInput = bRawGearUpInput;
-				AsyncInput->ControlInputs.GearDownInput = bRawGearDownInput;
-				AsyncInput->ControlInputs.TransmissionType = TransmissionType;
-
-				AsyncInput->ControlInputs.ParkingEnabled = bParkEnabled;
-
-				// debug feature to limit the vehicles top speed
-				if ((GVehicleDebugParams.SetMaxMPH > 0.f) && (FMath::Abs(ThrottleInput) > 0.0f) && FMath::Abs(GetForwardSpeedMPH()) >= GVehicleDebugParams.SetMaxMPH)
+
+				if (!bUsingNetworkPhysicsPrediction)
 				{
-					AsyncInput->ControlInputs.ThrottleInput = 0.1f;
+					AsyncInput->ControlInputs.ThrottleInput = ThrottleInputRate.CalcControlFunction(ThrottleInput);
+					AsyncInput->ControlInputs.BrakeInput = BrakeInputRate.CalcControlFunction(BrakeInput);
+					AsyncInput->ControlInputs.SteeringInput = SteeringInputRate.CalcControlFunction(SteeringInput);
+					AsyncInput->ControlInputs.HandbrakeInput = HandbrakeInput;
+					AsyncInput->ControlInputs.RollInput = RollInputRate.CalcControlFunction(RollInput);
+					AsyncInput->ControlInputs.PitchInput = PitchInputRate.CalcControlFunction(PitchInput);
+					AsyncInput->ControlInputs.YawInput = YawInputRate.CalcControlFunction(YawInput);
+					AsyncInput->ControlInputs.GearUpInput = bRawGearUpInput;
+					AsyncInput->ControlInputs.GearDownInput = bRawGearDownInput;
+					AsyncInput->ControlInputs.TransmissionType = TransmissionType;
+
+					AsyncInput->ControlInputs.ParkingEnabled = bParkEnabled;
+
+					// debug feature to limit the vehicles top speed
+					if ((GVehicleDebugParams.SetMaxMPH > 0.f) && (FMath::Abs(ThrottleInput) > 0.0f) && FMath::Abs(GetForwardSpeedMPH()) >= GVehicleDebugParams.SetMaxMPH)
+					{
+						AsyncInput->ControlInputs.ThrottleInput = 0.1f;
+					}
 				}
 
 				AsyncInput->GravityZ = GetGravityZ();
@@ -1918,6 +2392,57 @@ void UChaosVehicleMovementComponent::PutAllEnabledRigidBodiesToSleep()
 	}
 }
 
+void UChaosVehicleMovementComponent::OnRep_SetVehicleHistory()
+{
+	if (UWorld* World = GetWorld())
+	{
+		bool bHasController = false;
+		if (APawn* MyPawn = Cast<APawn>(GetOwner()))
+		{
+			if (APlayerController* PlayerController = Cast<APlayerController>(MyPawn->GetController()))
+			{
+				if (PlayerController->IsLocalController() && !PlayerController->IsNetMode(NM_DedicatedServer))
+				{
+					bHasController = true;
+				}
+			}
+		}
+
+		if (!bHasController)
+		{
+			const bool bControlInputPressed = ((ReplicatedDatas.InputsDatas.VehicleInputs.ThrottleInput >= GVehicleDebugParams.ControlInputWakeTolerance)
+				|| (ReplicatedDatas.InputsDatas.VehicleInputs.BrakeInput >= GVehicleDebugParams.ControlInputWakeTolerance)
+				|| (ReplicatedDatas.InputsDatas.VehicleInputs.RollInput >= GVehicleDebugParams.ControlInputWakeTolerance)
+				|| (ReplicatedDatas.InputsDatas.VehicleInputs.PitchInput >= GVehicleDebugParams.ControlInputWakeTolerance)
+				|| (ReplicatedDatas.InputsDatas.VehicleInputs.YawInput >= GVehicleDebugParams.ControlInputWakeTolerance)
+				|| (FMath::Abs(ReplicatedDatas.InputsDatas.VehicleInputs.SteeringInput - PrevSteeringInput) >= GVehicleDebugParams.ControlInputWakeTolerance));
+
+			if (bControlInputPressed)
+			{
+				SetSleeping(false);
+			}
+		}
+
+		APawn* LocalPawn = Cast<APawn>(GetOwner());
+		if (APlayerController* PlayerController = (LocalPawn && LocalPawn->GetController()) ?
+			Cast<APlayerController>(LocalPawn->GetController()) : World->GetFirstPlayerController())
+		{
+			const int32 HistoryFrame = ReplicatedDatas.ServerFrame - PlayerController->GetLocalToServerAsyncPhysicsTickOffset();
+			const FVehicleReplicatedDatas ServerDatas = ReplicatedDatas;
+
+			const int32 CurrentFrame = PlayerController->GetAsyncPhysicsTimestamp().LocalFrame;
+			PlayerController->ExecuteAsyncPhysicsCommand(PlayerController->GetAsyncPhysicsTimestamp(), this, [this, ServerDatas, HistoryFrame, bHasController, CurrentFrame]
+				{
+					if (HistoryFrame >= 0)
+					{
+						VehicleSimulationPT->InputsHistory->RecordInputs(HistoryFrame, ServerDatas.InputsDatas);
+						VehicleSimulationPT->StatesHistory->RecordStates(HistoryFrame, ServerDatas.StatesDatas);
+					}
+				}, false);
+		}
+	}
+}
+
 
 #undef LOCTEXT_NAMESPACE
 
diff --git a/Engine/Plugins/Experimental/ChaosVehiclesPlugin/Source/ChaosVehicles/Private/ChaosWheeledVehicleMovementComponent.cpp b/Engine/Plugins/Experimental/ChaosVehiclesPlugin/Source/ChaosVehicles/Private/ChaosWheeledVehicleMovementComponent.cpp
--- a/Engine/Plugins/Experimental/ChaosVehiclesPlugin/Source/ChaosVehicles/Private/ChaosWheeledVehicleMovementComponent.cpp
+++ b/Engine/Plugins/Experimental/ChaosVehiclesPlugin/Source/ChaosVehicles/Private/ChaosWheeledVehicleMovementComponent.cpp
@@ -178,12 +178,10 @@ void UChaosWheeledVehicleSimulation::TickVehicle(UWorld* WorldIn, float DeltaTim
 	UChaosVehicleSimulation::TickVehicle(WorldIn, DeltaTime, InputData, OutputData, Handle);
 }
 
-void UChaosWheeledVehicleSimulation::UpdateSimulation(float DeltaTime, const FChaosVehicleDefaultAsyncInput& InputData, Chaos::FRigidBodyHandle_Internal* Handle)
-{
-	SCOPE_CYCLE_COUNTER(STAT_ChaosVehicle_UpdateSimulation);
 
-	// Inherit common vehicle simulation stages ApplyAerodynamics, ApplyTorqueControl, etc
-	UChaosVehicleSimulation::UpdateSimulation(DeltaTime, InputData, Handle);
+void UChaosWheeledVehicleSimulation::UpdateState(float DeltaTime, const FChaosVehicleDefaultAsyncInput& InputData, Chaos::FRigidBodyHandle_Internal* Handle)
+{
+	UChaosVehicleSimulation::UpdateState(DeltaTime, InputData, Handle);
 
 	if (CanSimulate() && Handle)
 	{
@@ -223,7 +221,6 @@ void UChaosWheeledVehicleSimulation::UpdateSimulation(float DeltaTime, const FCh
 				WheelState.CaptureState(WheelIdx, PVehicle->Suspension[WheelIdx].GetLocalRestingPosition(), Handle);
 			}
 		}
-
 		///////////////////////////////////////////////////////////////////////
 		// Suspension Raycast
 
@@ -257,11 +254,18 @@ void UChaosWheeledVehicleSimulation::UpdateSimulation(float DeltaTime, const FCh
 			}
 		}
 		VehicleState.bAllWheelsOnGround = (VehicleState.NumWheelsOnGround == PVehicle->Wheels.Num());
+	}
+}
 
-		///////////////////////////////////////////////////////////////////////
-		// Input
-		ApplyInput(InputData.ControlInputs, DeltaTime);
+void UChaosWheeledVehicleSimulation::UpdateSimulation(float DeltaTime, const FChaosVehicleDefaultAsyncInput& InputData, Chaos::FRigidBodyHandle_Internal* Handle)
+{
+	SCOPE_CYCLE_COUNTER(STAT_ChaosVehicle_UpdateSimulation);
+	
+	// Inherit common vehicle simulation stages ApplyAerodynamics, ApplyTorqueControl, etc
+	UChaosVehicleSimulation::UpdateSimulation(DeltaTime, InputData, Handle);
 
+	if (CanSimulate() && Handle)
+	{
 		///////////////////////////////////////////////////////////////////////
 		// Engine/Transmission
 		if (!GWheeledVehicleDebugParams.DisableSuspensionForces && PVehicle->bMechanicalSimEnabled)
@@ -297,7 +301,6 @@ void UChaosWheeledVehicleSimulation::UpdateSimulation(float DeltaTime, const FCh
 		}
 #endif
 	}
-
 }
 
 bool UChaosWheeledVehicleSimulation::ContainsTraces(const FBox& Box, const TArray<Chaos::FSuspensionTrace>& SuspensionTrace)
@@ -527,6 +530,15 @@ void UChaosWheeledVehicleSimulation::ApplyWheelFrictionForces(float DeltaTime)
 			FMatrix Mat = FMatrix(GroundXVector, GroundYVector, GroundZVector, VehicleState.VehicleWorldTransform.GetLocation());
 			FVector FrictionForceVector = Mat.TransformVector(FrictionForceLocal);
 
+#if !(UE_BUILD_SHIPPING || UE_BUILD_TEST)
+			if (Chaos::FPhysicsSolverBase::IsNetworkPhysicsPredictionEnabled() && Chaos::FPhysicsSolverBase::CanDebugNetworkPhysicsPrediction())
+			{
+				UE_LOG(LogVehicle, Log, TEXT("Friction Force Global = %s | Transform = %s | Local = %s | Steering Angle = %f | Velocity = %s | Rotation Angle = %f | Friction Force = %s"),
+					*FrictionForceVector.ToString(), *Mat.ToString(), *FrictionForceLocal.ToString(), 
+					SteerAngleDegrees, *WheelState.LocalWheelVelocity[WheelIdx].ToString(), RotationAngle, *PWheel.GetForceFromFriction().ToString());
+			}
+#endif
+
 			check(PWheel.InContact());
 			if (PVehicle->bLegacyWheelFrictionPosition)
 			{
@@ -536,7 +548,7 @@ void UChaosWheeledVehicleSimulation::ApplyWheelFrictionForces(float DeltaTime)
 			{
 				AddForceAtPosition(FrictionForceVector, HitResult.ImpactPoint);
 			}
-
+		
 #if !(UE_BUILD_SHIPPING || UE_BUILD_TEST)
 			if (GWheeledVehicleDebugParams.ShowWheelForces)
 			{
@@ -722,6 +734,15 @@ void UChaosWheeledVehicleSimulation::ProcessSteering(const FControlInputs& Contr
 
 			float SteeringAngle = ControlInputs.SteeringInput * SpeedScale;
 
+#if !(UE_BUILD_SHIPPING || UE_BUILD_TEST)
+			if (Chaos::FPhysicsSolverBase::IsNetworkPhysicsPredictionEnabled() && Chaos::FPhysicsSolverBase::CanDebugNetworkPhysicsPrediction())
+			{
+				UE_LOG(LogVehicle, Log, TEXT("Process Sterring Input = %f | Spped Scale = %f | Wheel Side = %f"),
+					ControlInputs.SteeringInput, SpeedScale, PVehicle->GetSuspension(WheelIdx).GetLocalRestingPosition().Y);
+			}
+#endif
+
+
 			if (FMath::Abs(GWheeledVehicleDebugParams.SteeringOverride) > 0.01f)
 			{
 				SteeringAngle = PWheel.MaxSteeringAngle * GWheeledVehicleDebugParams.SteeringOverride;
@@ -818,11 +839,15 @@ void UChaosWheeledVehicleSimulation::ApplyInput(const FControlInputs& ControlInp
 
 bool UChaosWheeledVehicleSimulation::IsWheelSpinning() const
 {
-	for (auto& Wheel : PVehicle->Wheels)
-	{
-		if (Wheel.IsSlipping())
+	if(!Chaos::FPhysicsSolverBase::IsNetworkPhysicsPredictionEnabled())
+	{ 
+		// Disable it for now since it is breaking determinism
+		for (auto& Wheel : PVehicle->Wheels)
 		{
-			return true;
+			if (Wheel.IsSlipping())
+			{
+				return true;
+			}
 		}
 	}
 
diff --git a/Engine/Plugins/Experimental/ChaosVehiclesPlugin/Source/ChaosVehicles/Public/ChaosVehicleMovementComponent.h b/Engine/Plugins/Experimental/ChaosVehiclesPlugin/Source/ChaosVehicles/Public/ChaosVehicleMovementComponent.h
--- a/Engine/Plugins/Experimental/ChaosVehiclesPlugin/Source/ChaosVehicles/Public/ChaosVehicleMovementComponent.h
+++ b/Engine/Plugins/Experimental/ChaosVehiclesPlugin/Source/ChaosVehicles/Public/ChaosVehicleMovementComponent.h
@@ -17,6 +17,7 @@
 #include "DeferredForces.h"
 #include "ChaosVehicleManagerAsyncCallback.h"
 #include "CollisionQueryParams.h"
+#include "RewindData.h"
 
 #include "ChaosVehicleMovementComponent.generated.h"
 
@@ -29,10 +30,13 @@ class UCanvas;
 struct FChaosVehicleAsyncInput;
 struct FChaosVehicleManagerAsyncOutput;
 
-
-struct CHAOSVEHICLES_API FControlInputs
+/** Vehicle inputs from the player controller */
+USTRUCT()
+struct CHAOSVEHICLES_API FVehicleInputs
 {
-	FControlInputs()
+	GENERATED_USTRUCT_BODY()
+
+		FVehicleInputs()
 		: SteeringInput(0.f)
 		, ThrottleInput(0.f)
 		, BrakeInput(0.f)
@@ -40,36 +44,62 @@ struct CHAOSVEHICLES_API FControlInputs
 		, RollInput(0.f)
 		, YawInput(0.f)
 		, HandbrakeInput(0.f)
-		, ParkingEnabled(false)
-		, TransmissionType(Chaos::ETransmissionType::Automatic)
-		, GearUpInput(false)
-		, GearDownInput(false)
-	{
+	{}
 
+	/** Copy the parent class properties from another instance */
+	void CopyVehicleInputs(const FVehicleInputs& VehicleInputs)
+	{
+		SteeringInput = VehicleInputs.SteeringInput;
+		ThrottleInput = VehicleInputs.ThrottleInput;
+		BrakeInput = VehicleInputs.BrakeInput;
+		PitchInput = VehicleInputs.PitchInput;
+		RollInput = VehicleInputs.RollInput;
+		YawInput = VehicleInputs.YawInput;
+		HandbrakeInput = VehicleInputs.HandbrakeInput;
 	}
 
 	// Steering output to physics system. Range -1...1
+	UPROPERTY()
 	float SteeringInput;
 
 	// Accelerator output to physics system. Range 0...1
+	UPROPERTY()
 	float ThrottleInput;
 
 	// Brake output to physics system. Range 0...1
+	UPROPERTY()
 	float BrakeInput;
 
 	// Body Pitch output to physics system. Range -1...1
+	UPROPERTY()
 	float PitchInput;
 
 	// Body Roll output to physics system. Range -1...1
+	UPROPERTY()
 	float RollInput;
 
 	// Body Yaw output to physics system. Range -1...1
+	UPROPERTY()
 	float YawInput;
 
 	// Handbrake output to physics system. Range 0...1
+	UPROPERTY()
 	float HandbrakeInput;
+};
+
+USTRUCT()
+struct CHAOSVEHICLES_API FControlInputs : public FVehicleInputs
+{
+	GENERATED_USTRUCT_BODY()
+
+		FControlInputs()
+		: FVehicleInputs()
+		, ParkingEnabled(false)
+		, TransmissionType(Chaos::ETransmissionType::Automatic)
+		, GearUpInput(false)
+		, GearDownInput(false)
+	{}
 
-	// ParkingEnabled
 	bool ParkingEnabled;
 
 	Chaos::ETransmissionType TransmissionType;
@@ -78,6 +108,135 @@ struct CHAOSVEHICLES_API FControlInputs
 	bool GearDownInput;
 };
 
+USTRUCT()
+struct CHAOSVEHICLES_API FVehicleHistoryDatas
+{
+	GENERATED_USTRUCT_BODY()
+
+	int32 DatasFrame = INDEX_NONE;
+
+	int32 SmoothFrame = INDEX_NONE;
+};
+
+/** Vehicle Inputs datas that will be used in the Inputs history to be applied while simulating */
+USTRUCT()
+struct CHAOSVEHICLES_API FVehicleInputsDatas : public FVehicleHistoryDatas
+{
+	GENERATED_USTRUCT_BODY()
+
+	/** List of incoming control inputs coming from the local client */
+	UPROPERTY()
+	FVehicleInputs VehicleInputs;
+
+	/** Transmission change time that could be set/changed from GT */
+	UPROPERTY()
+	float TransmissionChangeTime;
+
+	/** Transmission current gear that could be set/changed from GT */
+	UPROPERTY()
+	int32 TransmissionCurrentGear;
+
+	/** Transmission target gear that could be set/changed from GT */
+	UPROPERTY()
+	int32 TransmissionTargetGear;
+};
+
+/** Vehicle states datas that will be used in the states history to rewind the simulation at some point inn time */
+USTRUCT()
+struct CHAOSVEHICLES_API FVehicleStatesDatas : public FVehicleHistoryDatas
+{
+	GENERATED_USTRUCT_BODY()
+
+	/** Vehicle state last velocity */
+	UPROPERTY()
+	FVector StateLastVelocity;
+
+	/** Angular velocity for each wheels */
+	UPROPERTY()
+	TArray<float> WheelsOmega;
+
+	/** Angular position for each wheels */
+	UPROPERTY()
+	TArray<float> WheelsAngularPosition;
+
+	/** Suspension latest displacement to be used while simulating */
+	UPROPERTY()
+	TArray<float> SuspensionLastDisplacement;
+
+	/** Suspension latest spring length to be used while simulating  */
+	UPROPERTY()
+	TArray<float> SuspensionLastSpringLength;
+
+	/** Suspension averaged length for smoothing */
+	UPROPERTY()
+	TArray<float> SuspensionAveragedLength;
+
+	/** Suspension averaged count for smoothing */
+	UPROPERTY()
+	TArray<int32> SuspensionAveragedCount;
+
+	/** Suspension averaged number for smoothing */
+	UPROPERTY()
+	TArray<int32> SuspensionAveragedNum;
+
+	/** Engine angular velocity */
+	UPROPERTY()
+	float EngineOmega;
+};
+
+/** Vehicle states history that will be registered in the rewind datas */
+struct CHAOSVEHICLES_API FVehicleStatesHistory : public Chaos::FStatesRewindHistory<FVehicleStatesDatas>
+{
+	using Super = Chaos::FStatesRewindHistory<FVehicleStatesDatas>;
+	using DatasType = FVehicleStatesDatas;
+
+	FVehicleStatesHistory(const int32 FrameCount, const bool bIsHistoryLocal, class UChaosVehicleSimulation* ChaosSimulation);
+
+	/** Rewind the vehicle state at some point in time */
+	virtual bool RewindStates(const int32 RewindFrame, const bool bResetSolver) override;
+
+	/** Record the vehicle state from the replicated datas */
+	bool RecordStates(const int32 RecordFrame, const DatasType& ReplicatedHistory);
+
+	/** Replicate the states from the vehicle simulation*/
+	void ReplicateStates(DatasType& ReplicatedDatas);
+
+	virtual ~FVehicleStatesHistory() override;
+
+private:
+
+	class UChaosVehicleSimulation* VehicleSimulation = nullptr;
+};
+
+/** Vehicle Inputs history that will be registered in the rewind datas */
+struct CHAOSVEHICLES_API FVehicleInputsHistory : public Chaos::FInputsRewindHistory<FVehicleInputsDatas>
+{
+	using Super = Chaos::FInputsRewindHistory<FVehicleInputsDatas>;
+	using DatasType = FVehicleInputsDatas;
+
+	FVehicleInputsHistory(const int32 FrameCount, const bool bIsHistoryLocal, class UChaosVehicleSimulation* ChaosSimulation);
+
+	/** Apply the Inputs at some point in time */
+	virtual bool ApplyInputs(const int32 ApplyFrame, const bool bResetSolver) override;
+
+	/** Record the vehicle Inputs from the replicated datas */
+	bool RecordInputs(const int32 RecordFrame, const DatasType& ReplicatedHistory);
+
+	/** Replicate the Inputs from the vehicle simulation + external inputs */
+	void ReplicateInputs(
+		const float SteeringInput, const float ThrottleInput,
+		const float BrakeInput, const float HandbrakeInput, const float RollInput,
+		const float PitchInput, const float YawInput, const int32 CurrentGear, DatasType& ReplicatedDatas);
+
+	/** Smooth Inputs inputs (WIP) */
+	void SmoothInputs(const int32 ApplyFrame, const int32 ValidFrame, const bool bApplySmoothing);
+
+	virtual ~FVehicleInputsHistory() override;
+
+private:
+
+	class UChaosVehicleSimulation* VehicleSimulation = nullptr;
+};
 
 struct CHAOSVEHICLES_API FVehicleDebugParams
 {
@@ -104,64 +263,50 @@ struct CHAOSVEHICLES_API FVehicleDebugParams
 struct FBodyInstance;
 
 USTRUCT()
-struct CHAOSVEHICLES_API FVehicleReplicatedState
+struct CHAOSVEHICLES_API FVehicleReplicatedState : public FVehicleInputs
 {
 	GENERATED_USTRUCT_BODY()
 
-	FVehicleReplicatedState()
+	FVehicleReplicatedState() : FVehicleInputs()
 	{
-		SteeringInput = 0.f;
-		ThrottleInput = 0.f;
-		BrakeInput = 0.f;
-		PitchInput = 0.f;
-		RollInput = 0.f;
-		YawInput = 0.f;
-		HandbrakeInput = 0.f;
 		TargetGear = 0;
 		ThrottleUp = 0.f;
 		ThrottleDown = 0.f;
 	}
 
-	// input replication: steering
-	UPROPERTY()
-	float SteeringInput;
-
-	// input replication: throttle
-	UPROPERTY()
-	float ThrottleInput;
-
-	// input replication: brake
+	// state replication: gear
 	UPROPERTY()
-	float BrakeInput;
+	int32 TargetGear;
 
-	// input replication: body pitch
+	// input replication: increase throttle
 	UPROPERTY()
-	float PitchInput;
+	float ThrottleUp;
 
-	// input replication: body roll
+	// input replication: decrease throttle
 	UPROPERTY()
-	float RollInput;
+	float ThrottleDown;
+};
 
-	// input replication: body yaw
-	UPROPERTY()
-	float YawInput;
+/** Group of datas that needs to be replicated across all clients for rewind/resim */
+USTRUCT()
+struct CHAOSVEHICLES_API FVehicleReplicatedDatas
+{
+	GENERATED_USTRUCT_BODY()
 
-	// input replication: handbrake
-	UPROPERTY()
-	float HandbrakeInput;
+	FVehicleReplicatedDatas()
+	{}
 
-	// state replication: gear
+	/** Server frame at which the datas have been recorded */
 	UPROPERTY()
-	int32 TargetGear;
+	int32 ServerFrame = INDEX_NONE;
 
-	// input replication: increase throttle
+	/** Group of inputs to replicate */
 	UPROPERTY()
-	float ThrottleUp;
+	FVehicleStatesDatas StatesDatas;
 
-	// input replication: decrease throttle
+	/** Group of inputs to replicate */
 	UPROPERTY()
-	float ThrottleDown;
-
+	FVehicleInputsDatas InputsDatas;
 };
 
 struct CHAOSVEHICLES_API FWheelTraceParams
@@ -173,7 +318,7 @@ struct CHAOSVEHICLES_API FWheelTraceParams
 struct CHAOSVEHICLES_API FChaosVehicleDefaultAsyncInput : public FChaosVehicleAsyncInput
 {
 	float GravityZ;
-	FControlInputs ControlInputs;
+	mutable FControlInputs ControlInputs;
 	mutable FCollisionQueryParams TraceParams;
 	mutable FCollisionResponseContainer TraceCollisionResponse;
 	mutable TArray<FWheelTraceParams> WheelTraceParams;
@@ -181,9 +326,8 @@ struct CHAOSVEHICLES_API FChaosVehicleDefaultAsyncInput : public FChaosVehicleAs
 	FChaosVehicleDefaultAsyncInput();
 
 	virtual TUniquePtr<FChaosVehicleAsyncOutput> Simulate(UWorld* World, const float DeltaSeconds, const float TotalSeconds, bool& bWakeOut) const override;
-	
-	virtual void ApplyDeferredForces(Chaos::FRigidBodyHandle_Internal* RigidHandle) const override;
 
+	virtual void ApplyDeferredForces(Chaos::FRigidBodyHandle_Internal* RigidHandle) const override;
 };
 
 
@@ -193,7 +337,7 @@ struct CHAOSVEHICLES_API FVehicleTorqueControlConfig
 public:
 	GENERATED_USTRUCT_BODY()
 
-	FVehicleTorqueControlConfig()
+		FVehicleTorqueControlConfig()
 	{
 		InitDefaults();
 	}
@@ -208,7 +352,7 @@ public:
 
 	UPROPERTY(EditAnywhere, Category = Setup)
 	float YawFromSteering;
-	
+
 	UPROPERTY(EditAnywhere, Category = Setup)
 	float YawFromRollTorqueScaling;
 
@@ -416,9 +560,9 @@ struct CHAOSVEHICLES_API FVehicleState
 		, VehicleWorldAngularVelocity(FVector::ZeroVector)
 		, VehicleWorldCOM(FVector::ZeroVector)
 		, WorldVelocityNormal(FVector::ZeroVector)
-		, VehicleUpAxis(FVector(0.f,0.f,1.f))
-		, VehicleForwardAxis(FVector(1.f,0.f,0.f))
-		, VehicleRightAxis(FVector(0.f,1.f,0.f))
+		, VehicleUpAxis(FVector(0.f, 0.f, 1.f))
+		, VehicleForwardAxis(FVector(1.f, 0.f, 0.f))
+		, VehicleRightAxis(FVector(0.f, 1.f, 0.f))
 		, LocalAcceleration(FVector::ZeroVector)
 		, LocalGForce(FVector::ZeroVector)
 		, LastFrameVehicleLocalVelocity(FVector::ZeroVector)
@@ -477,7 +621,7 @@ struct CHAOSVEHICLES_API FVehicleInputRateConfig
 {
 	GENERATED_USTRUCT_BODY()
 
-	/** 
+	/**
 	 * Rate at which the input value rises
 	 */
 	UPROPERTY(EditAnywhere, Category=VehicleInputRate)
@@ -505,17 +649,17 @@ struct CHAOSVEHICLES_API FVehicleInputRateConfig
 	FVehicleInputRateConfig() : RiseRate(5.0f), FallRate(5.0f), InputCurveFunction(EInputFunctionType::LinearFunction) { }
 
 	/** Change an output value using max rise and fall rates */
-	float InterpInputValue( float DeltaTime, float CurrentValue, float NewValue ) const
+	float InterpInputValue(float DeltaTime, float CurrentValue, float NewValue) const
 	{
 		const float DeltaValue = NewValue - CurrentValue;
 
 		// We are "rising" when DeltaValue has the same sign as CurrentValue (i.e. delta causes an absolute magnitude gain)
 		// OR we were at 0 before, and our delta is no longer 0.
-		const bool bRising = (( DeltaValue > 0.0f ) == ( CurrentValue > 0.0f )) ||
-								(( DeltaValue != 0.f ) && ( CurrentValue == 0.f ));
+		const bool bRising = ((DeltaValue > 0.0f) == (CurrentValue > 0.0f)) ||
+			((DeltaValue != 0.f) && (CurrentValue == 0.f));
 
-		const float MaxDeltaValue = DeltaTime * ( bRising ? RiseRate : FallRate );
-		const float ClampedDeltaValue = FMath::Clamp( DeltaValue, -MaxDeltaValue, MaxDeltaValue );
+		const float MaxDeltaValue = DeltaTime * (bRising ? RiseRate : FallRate);
+		const float ClampedDeltaValue = FMath::Clamp(DeltaValue, -MaxDeltaValue, MaxDeltaValue);
 		return CurrentValue + ClampedDeltaValue;
 	}
 
@@ -531,7 +675,7 @@ struct CHAOSVEHICLES_API FVehicleInputRateConfig
 			if (UserCurve.GetRichCurveConst() && !UserCurve.GetRichCurveConst()->IsEmpty())
 			{
 				float Output = FMath::Clamp(UserCurve.GetRichCurveConst()->Eval(FMath::Abs(InputValue)), 0.0f, 1.0f);
-				return (InputValue < 0.f)? -Output : Output;
+				return (InputValue < 0.f) ? -Output : Output;
 			}
 			else
 			{
@@ -551,7 +695,7 @@ struct CHAOSVEHICLES_API FVehicleInputRateConfig
 			return InputValue;
 		}
 		break;
-		
+
 		}
 
 	}
@@ -709,7 +853,7 @@ struct CHAOSVEHICLES_API FVehicleThrustConfig
 		ThrustType = EVehicleThrustType::Fixed;
 		BoneName = NAME_None;
 		Offset = FVector::ZeroVector;
-		ThrustAxis = FVector(1,0,0);
+		ThrustAxis = FVector(1, 0, 0);
 		//ThrustCurve.GetRichCurve()->AddKey(0.f, 1.f);
 		//ThrustCurve.GetRichCurve()->AddKey(1.f, 1.f);
 		MaxThrustForce = 1000.0f;
@@ -717,7 +861,7 @@ struct CHAOSVEHICLES_API FVehicleThrustConfig
 	}
 
 private:
-	void FillThrusterSetup(const UChaosVehicleMovementComponent &MovementComponent);
+	void FillThrusterSetup(const UChaosVehicleMovementComponent& MovementComponent);
 
 	Chaos::FSimpleThrustConfig PThrusterConfig;
 
@@ -740,11 +884,19 @@ public:
 
 	virtual void TickVehicle(UWorld* WorldIn, float DeltaTime, const FChaosVehicleDefaultAsyncInput& InputData, FChaosVehicleAsyncOutput& OutputData, Chaos::FRigidBodyHandle_Internal* Handle);
 
+	/** Update the async inputs and the history ones*/
+	void SyncHistoryInputs(const FChaosVehicleDefaultAsyncInput& InputData, Chaos::FRigidBodyHandle_Internal* Handle);
+
+	/** Apply the defered forces onto the vehicles */
 	virtual void ApplyDeferredForces(Chaos::FRigidBodyHandle_Internal* Handle);
 
+	/** Update the vehicle state */
+	virtual void UpdateState(float DeltaTime, const FChaosVehicleDefaultAsyncInput& InputData, Chaos::FRigidBodyHandle_Internal* Handle);
+
 	/** Advance the vehicle simulation */
 	virtual void UpdateSimulation(float DeltaTime, const FChaosVehicleDefaultAsyncInput& InputData, Chaos::FRigidBodyHandle_Internal* Handle);
-	
+
+	/** Fill the vehicle output state */
 	virtual void FillOutputState(FChaosVehicleAsyncOutput& Output);
 
 	/** Are enough vehicle systems specified such that physics vehicle simulation is possible */
@@ -765,15 +917,18 @@ public:
 	/** Apply direct control over vehicle body rotation */
 	virtual void ApplyTorqueControl(float DeltaTime, const FChaosVehicleDefaultAsyncInput& InputData);
 
-
 	/** Add a force to this vehicle */
 	void AddForce(const FVector& Force, bool bAllowSubstepping = true, bool bAccelChange = false);
+
 	/** Add a force at a particular position (world space when bIsLocalForce = false, body space otherwise) */
 	void AddForceAtPosition(const FVector& Force, const FVector& Position, bool bAllowSubstepping = true, bool bIsLocalForce = false);
+
 	/** Add an impulse to this vehicle */
 	void AddImpulse(const FVector& Impulse, bool bVelChange);
+
 	/** Add an impulse to this vehicle and a particular world position */
 	void AddImpulseAtPosition(const FVector& Impulse, const FVector& Position);
+
 	/** Add a torque to this vehicle */
 	void AddTorqueInRadians(const FVector& Torque, bool bAllowSubstepping = true, bool bAccelChange = false);
 
@@ -797,13 +952,29 @@ public:
 
 	FDeferredForces DeferredForces;
 
+	/** Current control inputs that is being used on the PT */
+	FVehicleInputs VehicleInputs;
+
+	/** States history uses to rewind simulation */
+	TSharedPtr<FVehicleStatesHistory> StatesHistory;
+
+	/** Inputs history used during simulation*/
+	TSharedPtr<FVehicleInputsHistory> InputsHistory;
+
+	FVehicleInputRateConfig ThrottleInputRate;
+	FVehicleInputRateConfig SteeringInputRate;
+	FVehicleInputRateConfig BrakeInputRate;
+	FVehicleInputRateConfig HandbrakeInputRate;
+	FVehicleInputRateConfig RollInputRate;
+	FVehicleInputRateConfig PitchInputRate;
+	FVehicleInputRateConfig YawInputRate;
 };
 
 
 /**
  * Base component to handle the vehicle simulation for an actor.
  */
-UCLASS(Abstract, hidecategories=(PlanarMovement, "Components|Movement|Planar", Activation, "Components|Activation"))
+UCLASS(Abstract, hidecategories = (PlanarMovement, "Components|Movement|Planar", Activation, "Components|Activation"))
 class CHAOSVEHICLES_API UChaosVehicleMovementComponent : public UPawnMovementComponent
 {
 	friend struct FChaosVehicleDefaultAsyncInput;
@@ -811,7 +982,7 @@ class CHAOSVEHICLES_API UChaosVehicleMovementComponent : public UPawnMovementCom
 
 	GENERATED_UCLASS_BODY()
 
-//#todo: these 2 oddities seem out of place
+	//#todo: these 2 oddities seem out of place
 
 	/** If true, the brake and reverse controls will behave in a more arcade fashion where holding reverse also functions as brake. For a more realistic approach turn this off*/
 	UPROPERTY(EditAnywhere, Category = VehicleSetup)
@@ -864,13 +1035,13 @@ public:
 	float DebugDragMagnitude;
 
 	/** Scales the vehicle's inertia in each direction (forward, right, up) */
-	UPROPERTY(EditAnywhere, Category=VehicleSetup, AdvancedDisplay)
+	UPROPERTY(EditAnywhere, Category = VehicleSetup, AdvancedDisplay)
 	FVector InertiaTensorScale;
 
 	/** Option to apply some aggressive sleep logic, larger number is more agressive, 0 disables */
 	UPROPERTY(EditAnywhere, Category = VehicleSetup)
 	float SleepThreshold;
-	
+
 	/** Option to apply some aggressive sleep logic if slopes up Z is less than this value, i.e value = Cos(SlopeAngle) so 0.866 will sleep up to 30 degree slopes */
 	UPROPERTY(EditAnywhere, Category = VehicleSetup, meta = (ClampMin = "0.01", UIMin = "0.01", ClampMax = "1.0", UIMax = "1.0"))
 	float SleepSlopeLimit;
@@ -886,7 +1057,7 @@ public:
 	/** Arcade style direct control of vehicle rotation via torque force */
 	UPROPERTY(EditAnywhere, Category = ArcadeControl)
 	FVehicleTorqueControlConfig TorqueControl;
-	
+
 	/** Arcade style direct control of vehicle rotation via torque force */
 	UPROPERTY(EditAnywhere, Category = ArcadeControl)
 	FVehicleTargetRotationControlConfig TargetRotationControl;
@@ -961,7 +1132,7 @@ public:
 
 
 	/** Set the user input for the vehicle throttle [range 0 to 1] */
-	UFUNCTION(BlueprintCallable, Category="Game|Components|ChaosVehicleMovement")
+	UFUNCTION(BlueprintCallable, Category = "Game|Components|ChaosVehicleMovement")
 	void SetThrottleInput(float Throttle);
 
 	/** Increase the vehicle throttle position [throttle range normalized 0 to 1] */
@@ -975,9 +1146,9 @@ public:
 	/** Set the user input for the vehicle Brake [range 0 to 1] */
 	UFUNCTION(BlueprintCallable, Category = "Game|Components|ChaosVehicleMovement")
 	void SetBrakeInput(float Brake);
-	
+
 	/** Set the user input for the vehicle steering [range -1 to 1] */
-	UFUNCTION(BlueprintCallable, Category="Game|Components|ChaosVehicleMovement")
+	UFUNCTION(BlueprintCallable, Category = "Game|Components|ChaosVehicleMovement")
 	void SetSteeringInput(float Steering);
 
 	/** Set the user input for the vehicle pitch [range -1 to 1] */
@@ -993,7 +1164,7 @@ public:
 	void SetYawInput(float Yaw);
 
 	/** Set the user input for handbrake */
-	UFUNCTION(BlueprintCallable, Category="Game|Components|ChaosVehicleMovement")
+	UFUNCTION(BlueprintCallable, Category = "Game|Components|ChaosVehicleMovement")
 	void SetHandbrakeInput(bool bNewHandbrake);
 
 	/** Set the vehicle in park mode */
@@ -1005,23 +1176,23 @@ public:
 	void SetSleeping(bool bEnableSleep);
 
 	/** Set the user input for gear up */
-	UFUNCTION(BlueprintCallable, Category="Game|Components|ChaosVehicleMovement")
+	UFUNCTION(BlueprintCallable, Category = "Game|Components|ChaosVehicleMovement")
 	void SetChangeUpInput(bool bNewGearUp);
 
 	/** Set the user input for gear down */
-	UFUNCTION(BlueprintCallable, Category="Game|Components|ChaosVehicleMovement")
+	UFUNCTION(BlueprintCallable, Category = "Game|Components|ChaosVehicleMovement")
 	void SetChangeDownInput(bool bNewGearDown);
 
 	/** Set the user input for gear (-1 reverse, 0 neutral, 1+ forward)*/
-	UFUNCTION(BlueprintCallable, Category="Game|Components|ChaosVehicleMovement")
+	UFUNCTION(BlueprintCallable, Category = "Game|Components|ChaosVehicleMovement")
 	void SetTargetGear(int32 GearNum, bool bImmediate);
 
 	/** Set the flag that will be used to select auto-gears */
-	UFUNCTION(BlueprintCallable, Category="Game|Components|ChaosVehicleMovement")
+	UFUNCTION(BlueprintCallable, Category = "Game|Components|ChaosVehicleMovement")
 	void SetUseAutomaticGears(bool bUseAuto);
 
 	/** Set the flag that determines whether a controller is required to set control inputs */
-	UFUNCTION(BlueprintCallable, Category="Game|Components|ChaosVehicleMovement")
+	UFUNCTION(BlueprintCallable, Category = "Game|Components|ChaosVehicleMovement")
 	void SetRequiresControllerForInputs(bool bRequiresController);
 
 	/** Get current gear */
@@ -1037,7 +1208,7 @@ public:
 	bool GetUseAutoGears() const;
 
 	/** How fast the vehicle is moving forward */
-	UFUNCTION(BlueprintCallable, Category="Game|Components|ChaosVehicleMovement")
+	UFUNCTION(BlueprintCallable, Category = "Game|Components|ChaosVehicleMovement")
 	float GetForwardSpeed() const;
 
 	/** How fast the vehicle is moving forward */
@@ -1127,12 +1298,19 @@ public:
 	};
 	TArray<FAsyncOutputWrapper> OutputsWaitingOn;
 
+	UFUNCTION()
+	virtual void OnRep_SetVehicleHistory();
+
 protected:
 
 	// replicated state of vehicle 
 	UPROPERTY(Transient, Replicated)
 	FVehicleReplicatedState ReplicatedState;
 
+	// replicated state of vehicle 
+	UPROPERTY(Transient, ReplicatedUsing = OnRep_SetVehicleHistory)
+	FVehicleReplicatedDatas ReplicatedDatas;
+
 	// accumulator for RB replication errors 
 	float AngErrorAccumulator;
 
@@ -1189,15 +1367,15 @@ protected:
 	float HandbrakeInput;
 
 	// Bypass the need for a controller in order for the controls to be processed.
-	UPROPERTY(EditAnywhere, Category=VehicleInput)
+	UPROPERTY(EditAnywhere, Category = VehicleInput)
 	bool bRequiresControllerForInputs;
 
 	// How much to press the brake when the player has release throttle
-	UPROPERTY(EditAnywhere, Category=VehicleInput)
+	UPROPERTY(EditAnywhere, Category = VehicleInput)
 	float IdleBrakeInput;
 
 	// Auto-brake when absolute vehicle forward speed is less than this (cm/s)
-	UPROPERTY(EditAnywhere, Category=VehicleInput)
+	UPROPERTY(EditAnywhere, Category = VehicleInput)
 	float StopThreshold;
 
 	// Auto-brake when vehicle forward speed is opposite of player input by at least this much (cm/s)
@@ -1205,19 +1383,19 @@ protected:
 	float WrongDirectionThreshold;
 
 	// Rate at which input throttle can rise and fall
-	UPROPERTY(EditAnywhere, Category=VehicleInput, AdvancedDisplay)
+	UPROPERTY(EditAnywhere, Category = VehicleInput, AdvancedDisplay)
 	FVehicleInputRateConfig ThrottleInputRate;
 
 	// Rate at which input brake can rise and fall
-	UPROPERTY(EditAnywhere, Category=VehicleInput, AdvancedDisplay)
+	UPROPERTY(EditAnywhere, Category = VehicleInput, AdvancedDisplay)
 	FVehicleInputRateConfig BrakeInputRate;
 
 	// Rate at which input steering can rise and fall
-	UPROPERTY(EditAnywhere, Category=VehicleInput, AdvancedDisplay)
+	UPROPERTY(EditAnywhere, Category = VehicleInput, AdvancedDisplay)
 	FVehicleInputRateConfig SteeringInputRate;
 
 	// Rate at which input handbrake can rise and fall
-	UPROPERTY(EditAnywhere, Category=VehicleInput, AdvancedDisplay)
+	UPROPERTY(EditAnywhere, Category = VehicleInput, AdvancedDisplay)
 	FVehicleInputRateConfig HandbrakeInputRate;
 
 	// Rate at which input pitch can rise and fall
@@ -1262,7 +1440,7 @@ protected:
 	* Raw input won't be cleared, the vehicle may resume input based movement next frame.
 	*/
 	virtual void ClearInput();
-	
+
 	/**
 	* Clear all raw inputs to default values.
 	* Interpolated input won't be cleared, the vehicle will begin interpolating to no input.
@@ -1288,9 +1466,17 @@ protected:
 	void ServerUpdateState(float InSteeringInput, float InThrottleInput, float InBrakeInput
 			, float InHandbrakeInput, int32 InCurrentGear, float InRollInput, float InPitchInput, float InYawInput);
 
-
 	// Setup
 
+	/** New method of applying control inputs at a specific time on the server - for network prediction */
+	UFUNCTION(Server, Reliable)
+	void ServerApplyControls(float InSteeringInput, float InThrottleInput, float InBrakeInput, float InHandbrakeInput,
+			int32 InCurrentGear, float InRollInput, float InPitchInput, float InYawInput, const FAsyncPhysicsTimestamp& AsyncPhysicsTimestamp);
+
+	/** Apply Controls locally */
+	void ApplyControls_Imp(float InSteeringInput, float InThrottleInput, float InBrakeInput, float InHandbrakeInput,
+		int32 InCurrentGear, float InRollInput, float InPitchInput, float InYawInput, const FAsyncPhysicsTimestamp& AsyncPhysicsTimestamp, const bool bEnableResim = true);
+
 	/** Get our controller */
 	virtual AController* GetController() const override;
 
@@ -1371,10 +1557,11 @@ protected:
 	void WakeAllEnabledRigidBodies();
 	void PutAllEnabledRigidBodiesToSleep();
 
+
 	Chaos::FSimpleAerodynamicsConfig PAerodynamicsSetup;
 	int32 TargetGear;
 
 	float PrevSteeringInput;
 	float PrevReplicatedSteeringInput;
-
+	bool bUsingNetworkPhysicsPrediction;
 };
diff --git a/Engine/Plugins/Experimental/ChaosVehiclesPlugin/Source/ChaosVehicles/Public/ChaosWheeledVehicleMovementComponent.h b/Engine/Plugins/Experimental/ChaosVehiclesPlugin/Source/ChaosVehicles/Public/ChaosWheeledVehicleMovementComponent.h
--- a/Engine/Plugins/Experimental/ChaosVehiclesPlugin/Source/ChaosVehicles/Public/ChaosWheeledVehicleMovementComponent.h
+++ b/Engine/Plugins/Experimental/ChaosVehiclesPlugin/Source/ChaosVehicles/Public/ChaosWheeledVehicleMovementComponent.h
@@ -590,6 +590,9 @@ public:
 	/** Advance the vehicle simulation */
 	virtual void UpdateSimulation(float DeltaTime, const FChaosVehicleDefaultAsyncInput& InputData, Chaos::FRigidBodyHandle_Internal* Handle) override;
 
+	/** Update the vehicle state */
+	virtual void UpdateState(float DeltaTime, const FChaosVehicleDefaultAsyncInput& InputData, Chaos::FRigidBodyHandle_Internal* Handle) override;
+
 	virtual void FillOutputState(FChaosVehicleAsyncOutput& Output) override;
 
 	/** Are enough vehicle systems specified such that physics vehicle simulation is possible */
diff --git a/Engine/Source/Runtime/Engine/Classes/GameFramework/PawnMovementComponent.h b/Engine/Source/Runtime/Engine/Classes/GameFramework/PawnMovementComponent.h
--- a/Engine/Source/Runtime/Engine/Classes/GameFramework/PawnMovementComponent.h
+++ b/Engine/Source/Runtime/Engine/Classes/GameFramework/PawnMovementComponent.h
@@ -12,7 +12,27 @@
 #include "PawnMovementComponent.generated.h"
 
 class APawn;
+class UPrimitiveComponent;
 
+/**
+ * Enumerates all the action types that could be applied to the physics state
+ */
+UENUM()
+enum class EPhysicsStateAction : uint8
+{
+	AddForce,
+	AddTorque,
+	AddForceAtPosition,
+	AddLinearVelocity,
+	AddAngularVelocity,
+	AddVelocityAtPosition,
+	AddLinearImpulse,
+	AddAngularImpulse,
+	AddImpulseAtPosition,
+	AddAcceleration,
+	NumActions
+};
+  
 /** 
  * PawnMovementComponent can be used to update movement for an associated Pawn.
  * It also provides ways to accumulate and read directional input in a generic way (with AddInputVector(), ConsumeInputVector(), etc).
@@ -82,6 +102,10 @@ public:
 
 	virtual void OnTeleported() override;
 
+	/** Apply the async physics state action onto the solver */
+	void ApplyAsyncPhysicsStateAction(const UPrimitiveComponent* ActionComponent, const FName& BoneName, 
+		const EPhysicsStateAction ActionType, const FVector& ActionDatas, const FVector& ActionPosition = FVector::Zero());
+
 protected:
 
 	/** Pawn that owns this component. */
@@ -100,4 +124,20 @@ protected:
 public:
 
 	virtual void Serialize(FArchive& Ar) override;
+
+private:
+
+	/** Send the physics command to execute to the server */
+	UFUNCTION(Server, Reliable)
+	void ServerAsyncPhysicsStateAction(const UPrimitiveComponent* ActionComponent, const FName BoneName, const FAsyncPhysicsTimestamp Timestamp,
+			const EPhysicsStateAction ActionType, const FVector ActionDatas, const FVector ActionPosition = FVector::Zero());
+
+	/** Dispatch back the physics command onto all the clients */
+	UFUNCTION(NetMulticast, Reliable)
+	void MulticastAsyncPhysicsStateAction(const UPrimitiveComponent* ActionComponent, const FName BoneName, const FAsyncPhysicsTimestamp Timestamp,
+			const EPhysicsStateAction ActionType, const FVector ActionDatas, const FVector ActionPosition = FVector::Zero());
+
+	/** Execute the async physics state action at a given timestamp */
+	void ExecuteAsyncPhysicsStateAction(const UPrimitiveComponent* ActionComponent, const FName& BoneName, const FAsyncPhysicsTimestamp& Timestamp,
+		const EPhysicsStateAction ActionType, const FVector& ActionDatas, const FVector& ActionPosition = FVector::Zero());
 };
diff --git a/Engine/Source/Runtime/Engine/Classes/GameFramework/PlayerController.h b/Engine/Source/Runtime/Engine/Classes/GameFramework/PlayerController.h
--- a/Engine/Source/Runtime/Engine/Classes/GameFramework/PlayerController.h
+++ b/Engine/Source/Runtime/Engine/Classes/GameFramework/PlayerController.h
@@ -219,11 +219,18 @@ USTRUCT()
 struct FAsyncPhysicsTimestamp
 {
 	GENERATED_BODY()
+
+	/** Check if the server and local frames are valid */
+	bool IsValid() const
+	{
+		return ServerFrame != INDEX_NONE && LocalFrame != INDEX_NONE;
+	}
 	UPROPERTY()
 	int32 ServerFrame = INDEX_NONE;
 
 	UPROPERTY()
 	int32 LocalFrame = INDEX_NONE;
+
 };
 
 /**
@@ -2345,7 +2352,7 @@ public:
 		These commands are all run on the game thread. If you want to run on the physics thread see FPhysicsSolverBase::RegisterSimOneShotCallback
 		If OwningObject is not null this command will only fire as long as the object is still alive. This allows a lambda to still use the owning object's data for read/write (assuming data is designed for async physics tick)
 	*/
-	void ExecuteAsyncPhysicsCommand(const FAsyncPhysicsTimestamp& AsyncPhysicsTimestamp, UObject* OwningObject, const TFunction<void()>& Command);
+	void ExecuteAsyncPhysicsCommand(const FAsyncPhysicsTimestamp& AsyncPhysicsTimestamp, UObject* OwningObject, const TFunction<void()>& Command, const bool bEnableResim = true);
 
 	/** Generates a timestamp for the upcoming physics step (plus any pending time). Useful for synchronizing client and server events on a specific physics step */
 	FAsyncPhysicsTimestamp GetAsyncPhysicsTimestamp(float DeltaSeconds = 0.f);
diff --git a/Engine/Source/Runtime/Engine/Classes/PhysicsEngine/BodyInstance.h b/Engine/Source/Runtime/Engine/Classes/PhysicsEngine/BodyInstance.h
--- a/Engine/Source/Runtime/Engine/Classes/PhysicsEngine/BodyInstance.h
+++ b/Engine/Source/Runtime/Engine/Classes/PhysicsEngine/BodyInstance.h
@@ -8,6 +8,7 @@
 #include "Engine/EngineTypes.h"
 #include "CollisionQueryParams.h"
 #include "EngineDefines.h"
+#include "GameFramework/PlayerController.h"
 #if UE_ENABLE_INCLUDE_ORDER_DEPRECATED_IN_5_2
 #include "PhysxUserData.h"
 #endif
@@ -772,6 +773,8 @@ public:
 	/** Enable or disable inertia conditionin.  @see bInertiaConditioning */
 	void SetInertiaConditioningEnabled(bool bEnabled);
 
+	/** Apply async physics command onto the body instance*/
+	void ApplyAsyncPhysicsCommand(FAsyncPhysicsTimestamp TimeStamp, const bool bIsInternal, APlayerController* PlayerController, const TFunction<void()>& Command);
 
 	/** 
 	 * Set this body to either simulate or to be fixed/kinematic. 
@@ -799,9 +802,9 @@ public:
 	/** Add custom forces and torques on the body. The callback will be called more than once, if substepping enabled, for every substep.  */
 	void AddCustomPhysics(FCalculateCustomPhysics& CalculateCustomPhysics);
 	/** Add a force to this body */
-	void AddForce(const FVector& Force, bool bAllowSubstepping = true, bool bAccelChange = false);
+	void AddForce(const FVector& Force, bool bAllowSubstepping = true, bool bAccelChange = false, const FAsyncPhysicsTimestamp TimeStamp = FAsyncPhysicsTimestamp(), APlayerController* PlayerController = nullptr);
 	/** Add a force at a particular position (world space when bIsLocalForce = false, body space otherwise) */
-	void AddForceAtPosition(const FVector& Force, const FVector& Position, bool bAllowSubstepping = true, bool bIsLocalForce = false);
+	void AddForceAtPosition(const FVector& Force, const FVector& Position, bool bAllowSubstepping = true, bool bIsLocalForce = false, const FAsyncPhysicsTimestamp TimeStamp = FAsyncPhysicsTimestamp(), APlayerController* PlayerController = nullptr);
 	/** Clear accumulated forces on this body */
 	void ClearForces(bool bAllowSubstepping = true);
 
@@ -809,21 +812,21 @@ public:
 	void SetOneWayInteraction(bool InOneWayInteraction = true);
 
 	/** Add a torque to this body */
-	void AddTorqueInRadians(const FVector& Torque, bool bAllowSubstepping = true, bool bAccelChange = false);
+	void AddTorqueInRadians(const FVector& Torque, bool bAllowSubstepping = true, bool bAccelChange = false, const FAsyncPhysicsTimestamp TimeStamp = FAsyncPhysicsTimestamp(), APlayerController* PlayerController = nullptr);
 	/** Clear accumulated torques on this body */
 	void ClearTorques(bool bAllowSubstepping = true);
 
 	/** Add a rotational impulse to this body */
-	void AddAngularImpulseInRadians(const FVector& Impulse, bool bVelChange);
+	void AddAngularImpulseInRadians(const FVector& Impulse, bool bVelChange, const FAsyncPhysicsTimestamp TimeStamp = FAsyncPhysicsTimestamp(), APlayerController* PlayerController = nullptr);
 
 	/** Add an impulse to this body */
-	void AddImpulse(const FVector& Impulse, bool bVelChange);
+	void AddImpulse(const FVector& Impulse, bool bVelChange, const FAsyncPhysicsTimestamp TimeStamp = FAsyncPhysicsTimestamp(), APlayerController* PlayerController = nullptr);
 	
 	/** Add an impulse to this body and a particular world position */
-	void AddImpulseAtPosition(const FVector& Impulse, const FVector& Position);
+	void AddImpulseAtPosition(const FVector& Impulse, const FVector& Position, const FAsyncPhysicsTimestamp TimeStamp = FAsyncPhysicsTimestamp(), APlayerController* PlayerController = nullptr);
 
 	/** Add a velocity change impulse to this body and a particular world position */
-	void AddVelocityChangeImpulseAtLocation(const FVector& Impulse, const FVector& Position);
+	void AddVelocityChangeImpulseAtLocation(const FVector& Impulse, const FVector& Position, const FAsyncPhysicsTimestamp TimeStamp = FAsyncPhysicsTimestamp(), APlayerController* PlayerController = nullptr);
 
 	/** Set the linear velocity of this body */
 	void SetLinearVelocity(const FVector& NewVel, bool bAddToCurrent, bool bAutoWake = true);
diff --git a/Engine/Source/Runtime/Engine/Private/ActorReplication.cpp b/Engine/Source/Runtime/Engine/Private/ActorReplication.cpp
--- a/Engine/Source/Runtime/Engine/Private/ActorReplication.cpp
+++ b/Engine/Source/Runtime/Engine/Private/ActorReplication.cpp
@@ -259,6 +259,18 @@ void AActor::PostNetReceiveLocationAndRotation()
 
 void AActor::PostNetReceiveVelocity(const FVector& NewVelocity)
 {
+	static IConsoleVariable* EnableNetworkPhysicsCVar = IConsoleManager::Get().FindConsoleVariable(TEXT("np2.EnableNetworkPhysicsPrediction"));
+	if (EnableNetworkPhysicsCVar && (EnableNetworkPhysicsCVar->GetInt() == 1))
+	{
+		// required to keep client deterministic with server simulation
+		if (RootComponent && RootComponent->IsRegistered() && (NewVelocity != GetVelocity()))
+		{
+			if (UPrimitiveComponent* PrimitiveComponent = Cast<UPrimitiveComponent>(RootComponent))
+			{
+				PrimitiveComponent->SetPhysicsLinearVelocity(NewVelocity);
+			}
+		}
+	}
 }
 
 void AActor::PostNetReceivePhysicState()
diff --git a/Engine/Source/Runtime/Engine/Private/Components/ActorComponent.cpp b/Engine/Source/Runtime/Engine/Private/Components/ActorComponent.cpp
--- a/Engine/Source/Runtime/Engine/Private/Components/ActorComponent.cpp
+++ b/Engine/Source/Runtime/Engine/Private/Components/ActorComponent.cpp
@@ -1272,7 +1272,7 @@ void UActorComponent::DeferRemoveAsyncPhysicsTick()
 			{
 				SetAsyncPhysicsTickEnabled(false);
 			}
-		);
+		,false);
 	}
 }
 
diff --git a/Engine/Source/Runtime/Engine/Private/Components/CharacterMovementComponent.cpp b/Engine/Source/Runtime/Engine/Private/Components/CharacterMovementComponent.cpp
--- a/Engine/Source/Runtime/Engine/Private/Components/CharacterMovementComponent.cpp
+++ b/Engine/Source/Runtime/Engine/Private/Components/CharacterMovementComponent.cpp
@@ -7151,18 +7151,33 @@ void UCharacterMovementComponent::ApplyImpactPhysicsForces(const FHitResult& Imp
 				{
 					PushForceModificator *= BodyMass;
 				}
+				const bool bEnableNetworkPhysics = Chaos::FPhysicsSolverBase::IsNetworkPhysicsPredictionEnabled(); 
 
 				Force *= PushForceModificator;
 				const float ZeroVelocityTolerance = 1.0f;
 				if (ComponentVelocity.IsNearlyZero(ZeroVelocityTolerance))
 				{
 					Force *= InitialPushForceFactor;
-					ImpactComponent->AddImpulseAtLocation(Force, ForcePoint, Impact.BoneName);
+					if (!bEnableNetworkPhysics)
+					{
+						ImpactComponent->AddImpulseAtLocation(Force, ForcePoint, Impact.BoneName);
+					}
+					else
+					{
+						ApplyAsyncPhysicsStateAction(ImpactComponent, Impact.BoneName, EPhysicsStateAction::AddImpulseAtPosition, Force, ForcePoint);
+					}
 				}
 				else
 				{
 					Force *= PushForceFactor;
-					ImpactComponent->AddForceAtLocation(Force, ForcePoint, Impact.BoneName);
+					if (!bEnableNetworkPhysics)
+					{
+						ImpactComponent->AddForceAtLocation(Force, ForcePoint, Impact.BoneName);
+					}
+					else
+					{
+						ApplyAsyncPhysicsStateAction(ImpactComponent, Impact.BoneName, EPhysicsStateAction::AddForceAtPosition, Force, ForcePoint);
+					}
 				}
 			}
 		}
diff --git a/Engine/Source/Runtime/Engine/Private/Components/PawnMovementComponent.cpp b/Engine/Source/Runtime/Engine/Private/Components/PawnMovementComponent.cpp
--- a/Engine/Source/Runtime/Engine/Private/Components/PawnMovementComponent.cpp
+++ b/Engine/Source/Runtime/Engine/Private/Components/PawnMovementComponent.cpp
@@ -3,6 +3,9 @@
 #include "GameFramework/PawnMovementComponent.h"
 #include "GameFramework/Pawn.h"
 #include "GameFramework/PlayerController.h"
+#include "Engine/World.h"
+#include "Components/PrimitiveComponent.h"
+#include "PhysicsEngine/BodyInstance.h"
 
 #include UE_INLINE_GENERATED_CPP_BY_NAME(PawnMovementComponent)
 
@@ -119,3 +122,110 @@ void UPawnMovementComponent::MarkForClientCameraUpdate()
 		}
 	}
 }
+
+void UPawnMovementComponent::ApplyAsyncPhysicsStateAction(const UPrimitiveComponent* ActionComponent, const FName& BoneName, 
+	const EPhysicsStateAction ActionType, const FVector& ActionDatas, const FVector& ActionPosition)
+{
+	if (APawn* MyPawn = Cast<APawn>(GetOwner()))
+	{
+		if (APlayerController* PlayerController = Cast<APlayerController>(MyPawn->GetController()))
+		{
+			if (PlayerController->IsLocalController() && !PlayerController->IsNetMode(NM_DedicatedServer))
+			{
+				FAsyncPhysicsTimestamp TimeStamp = PlayerController->GetAsyncPhysicsTimestamp();
+				ExecuteAsyncPhysicsStateAction(ActionComponent, BoneName, TimeStamp, ActionType, ActionDatas, ActionPosition);
+
+				if (!PlayerController->IsNetMode(NM_Standalone))
+				{
+					if (!PlayerController->IsNetMode(NM_ListenServer))
+					{
+						ServerAsyncPhysicsStateAction(ActionComponent, BoneName, TimeStamp, ActionType, ActionDatas, ActionPosition);
+					}
+					else
+					{
+						MulticastAsyncPhysicsStateAction(ActionComponent, BoneName, TimeStamp, ActionType, ActionDatas, ActionPosition);
+					}
+				}
+			}
+		}
+	}
+}
+
+void UPawnMovementComponent::ServerAsyncPhysicsStateAction_Implementation(const UPrimitiveComponent* ActionComponent, const FName BoneName, const FAsyncPhysicsTimestamp Timestamp,
+	const EPhysicsStateAction ActionType, const FVector ActionDatas, const FVector ActionPosition)
+{
+	ExecuteAsyncPhysicsStateAction(ActionComponent, BoneName, Timestamp, ActionType, ActionDatas, ActionPosition);
+	MulticastAsyncPhysicsStateAction(ActionComponent, BoneName, Timestamp, ActionType, ActionDatas, ActionPosition);
+}
+
+void UPawnMovementComponent::MulticastAsyncPhysicsStateAction_Implementation(const UPrimitiveComponent* ActionComponent, const FName BoneName, const FAsyncPhysicsTimestamp Timestamp,
+	const EPhysicsStateAction ActionType, const FVector ActionDatas, const FVector ActionPosition)
+{
+	if (UWorld* World = GetWorld())
+	{
+		if (APlayerController* PlayerController = World->GetFirstPlayerController())
+		{
+			if (GetController() == nullptr)
+			{
+				FAsyncPhysicsTimestamp LocalTimeStamp = Timestamp;
+				LocalTimeStamp.LocalFrame = LocalTimeStamp.ServerFrame - PlayerController->GetLocalToServerAsyncPhysicsTickOffset();
+
+				ExecuteAsyncPhysicsStateAction(ActionComponent, BoneName, LocalTimeStamp, ActionType, ActionDatas, ActionPosition);
+			}
+		}
+	}
+}
+
+void UPawnMovementComponent::ExecuteAsyncPhysicsStateAction(const UPrimitiveComponent* ActionComponent, const FName& BoneName, const FAsyncPhysicsTimestamp& Timestamp,
+	const EPhysicsStateAction ActionType, const FVector& ActionDatas, const FVector& ActionPosition)
+{
+	if (FBodyInstance* BI = ActionComponent->GetBodyInstance(BoneName))
+	{
+		if (GetOwner()->IsNetMode(NM_Client) || GetOwner()->IsNetMode(NM_Standalone))
+		{
+			UE_LOG(LogTemp, Warning, TEXT("ApplyImpactAtLocationImp CLIENT Force = %s | Location = %s | LocalFrame = %d | ServerFrame = %d | ComponentName = %s | ComponentPtr = %d"), *ActionDatas.ToString(), *ActionPosition.ToString(), Timestamp.LocalFrame, Timestamp.ServerFrame, *ActionComponent->GetPathName(), ActionComponent);
+		}
+		else
+		{
+			UE_LOG(LogTemp, Warning, TEXT("ApplyImpactAtLocationImp SERVER Force = %s | Location = %s | LocalFrame = %d | ServerFrame = %d | ComponentName = %s | ComponentPtr = %d"), *ActionDatas.ToString(), *ActionPosition.ToString(), Timestamp.LocalFrame, Timestamp.ServerFrame, *ActionComponent->GetPathName(), ActionComponent);
+		}
+		APawn* LocalPawn = Cast<APawn>(GetOwner());
+		APlayerController* PlayerController = (LocalPawn && LocalPawn->GetController()) ?
+			Cast<APlayerController>(LocalPawn->GetController()) : GetWorld()->GetFirstPlayerController();
+
+		switch (ActionType)
+		{
+			case EPhysicsStateAction::AddImpulseAtPosition:
+				BI->AddImpulseAtPosition(ActionDatas, ActionPosition, Timestamp, PlayerController);
+				break;
+			case EPhysicsStateAction::AddForceAtPosition:
+				BI->AddForceAtPosition(ActionDatas, ActionPosition, true, false, Timestamp, PlayerController);
+				break;
+			case EPhysicsStateAction::AddVelocityAtPosition:
+				BI->AddVelocityChangeImpulseAtLocation(ActionDatas, ActionPosition, Timestamp, PlayerController);
+				break;
+			case EPhysicsStateAction::AddLinearImpulse:
+				BI->AddImpulse(ActionDatas, false, Timestamp, PlayerController);
+				break;
+			case EPhysicsStateAction::AddForce:
+				BI->AddForce(ActionDatas, true, false, Timestamp, PlayerController);
+				break;
+			case EPhysicsStateAction::AddAcceleration:
+				BI->AddForce(ActionDatas, true, true, Timestamp, PlayerController);
+				break;
+			case EPhysicsStateAction::AddLinearVelocity:
+				BI->AddImpulse(ActionDatas, true, Timestamp, PlayerController);
+				break;
+			case EPhysicsStateAction::AddAngularImpulse:
+				BI->AddAngularImpulseInRadians(ActionDatas, false, Timestamp, PlayerController);
+				break;
+			case EPhysicsStateAction::AddTorque:
+				BI->AddTorqueInRadians(ActionDatas, true, false, Timestamp, PlayerController);
+				break;
+			case EPhysicsStateAction::AddAngularVelocity:
+				BI->AddAngularImpulseInRadians(ActionDatas, true, Timestamp, PlayerController);
+				break;
+
+		}
+	}
+}
diff --git a/Engine/Source/Runtime/Engine/Private/PackageMapClient.cpp b/Engine/Source/Runtime/Engine/Private/PackageMapClient.cpp
--- a/Engine/Source/Runtime/Engine/Private/PackageMapClient.cpp
+++ b/Engine/Source/Runtime/Engine/Private/PackageMapClient.cpp
@@ -173,6 +173,13 @@ static FAutoConsoleVariableRef CVarQuantizeActorVelocityOnSpawn(
 	TEXT("When enabled, we will quantize Velocity for newly spawned actors to a single decimal of precision.")
 );
 
+static bool GbQuantizeActorRotationOnSpawn = true;
+static FAutoConsoleVariableRef CVarQuantizeActorRotationOnSpawn(
+	TEXT("net.QuantizeActorRotationOnSpawn"),
+	GbQuantizeActorRotationOnSpawn,
+	TEXT("When enabled, we will quantize Rotation for newly spawned actors to a single decimal of precision.")
+);
+
 static bool GbNetCheckNoLoadPackages = true;
 static FAutoConsoleVariableRef CVarNetCheckNoLoadPackages(
 	TEXT("net.CheckNoLoadPackages"),
@@ -632,7 +639,14 @@ bool UPackageMapClient::SerializeNewActor(FArchive& Ar, class UActorChannel *Cha
 			Ar.SerializeBits(&bSerializeRotation, 1);
 			if (bSerializeRotation)
 			{
-				Rotation.NetSerialize(Ar, this, SerSuccess);
+				if (GbQuantizeActorRotationOnSpawn)
+				{
+					Rotation.NetSerialize(Ar, this, SerSuccess);
+				} 
+				else
+				{
+					Ar << Rotation;
+				}
 			}
 			else
 			{
diff --git a/Engine/Source/Runtime/Engine/Private/Pawn.cpp b/Engine/Source/Runtime/Engine/Private/Pawn.cpp
--- a/Engine/Source/Runtime/Engine/Private/Pawn.cpp
+++ b/Engine/Source/Runtime/Engine/Private/Pawn.cpp
@@ -1163,6 +1163,11 @@ void APawn::GetMoveGoalReachTest(const AActor* MovingActor, const FVector& MoveO
 
 void APawn::PostNetReceiveVelocity(const FVector& NewVelocity)
 {
+	static IConsoleVariable* EnableNetworkPhysicsCVar = IConsoleManager::Get().FindConsoleVariable(TEXT("np2.EnableNetworkPhysicsPrediction"));
+	if (EnableNetworkPhysicsCVar && EnableNetworkPhysicsCVar->GetInt() == 1)
+	{
+		Super::PostNetReceiveVelocity(NewVelocity);
+	}
 	if (GetLocalRole() == ROLE_SimulatedProxy)
 	{
 		UMovementComponent* const MoveComponent = GetMovementComponent();
diff --git a/Engine/Source/Runtime/Engine/Private/PhysicsEngine/AsyncPhysicsInputComponent.cpp b/Engine/Source/Runtime/Engine/Private/PhysicsEngine/AsyncPhysicsInputComponent.cpp
--- a/Engine/Source/Runtime/Engine/Private/PhysicsEngine/AsyncPhysicsInputComponent.cpp
+++ b/Engine/Source/Runtime/Engine/Private/PhysicsEngine/AsyncPhysicsInputComponent.cpp
@@ -8,6 +8,7 @@
 #include "GameFramework/Pawn.h"
 #include "GameFramework/PlayerController.h"
 #include "PBDRigidsSolver.h"
+#include "PhysicsReplication.h"
 #include "Physics/Experimental/PhysScene_Chaos.h"
 #include "PhysicsProxy/SingleParticlePhysicsProxy.h"
 
@@ -97,8 +98,7 @@ struct FAsyncPhysicsInputRewindCallback : public Chaos::IRewindCallback
 		{
 			if (APlayerController* PC = World->GetFirstPlayerController())
 			{
-				static IConsoleVariable* EnableNetworkPhysicsCVar = IConsoleManager::Get().FindConsoleVariable(TEXT("np2.EnableNetworkPhysicsPrediction"));
-				if (EnableNetworkPhysicsCVar && EnableNetworkPhysicsCVar->GetInt() == 1)
+				if (Chaos::FPhysicsSolverBase::IsNetworkPhysicsPredictionEnabled())
 				{
 					const int32 LocalToServerOffset = PC->GetLocalToServerAsyncPhysicsTickOffset();
 					ServerFrame = PhysicsStep + LocalToServerOffset;
@@ -116,6 +116,23 @@ struct FAsyncPhysicsInputRewindCallback : public Chaos::IRewindCallback
 		DispatchPhysicsTick.Broadcast(PhysicsStep, NumSteps, ServerFrame);
 	}
 
+	virtual int32 TriggerRewindIfNeeded_Internal(int32 LatestStepCompleted)
+	{
+		int32 CachedClientFrame = INDEX_NONE;
+		if (World->GetNetMode() == NM_Client && World->GetPhysicsScene())
+		{
+			if (FPhysScene_Chaos* Scene = static_cast<FPhysScene_Chaos*>(World->GetPhysicsScene()))
+			{
+				if(const FPhysicsReplication* PhysicsReplication  = Scene->GetPhysicsReplication())
+				{
+					CachedClientFrame = Scene->GetPhysicsReplication()->GetResimFrame();
+				}
+			}
+		}
+
+		return CachedClientFrame;
+	}
+
 	virtual void ProcessInputs_External(int32 PhysicsStep, const TArray<Chaos::FSimCallbackInputAndObject>& SimCallbackInputs)
 	{
 		CachedServerFrame = PhysicsStep;
@@ -138,10 +155,14 @@ struct FAsyncPhysicsInputRewindCallback : public Chaos::IRewindCallback
 				// -----------------------------------------------------------------
 				APlayerController::FClientFrameInfo& ClientFrameInfo = PC->GetClientFrameInfo();
 
-				if (ClientFrameInfo.LastProcessedInputFrame != INDEX_NONE)
+				if (Chaos::FPhysicsSolverBase::IsNetworkPhysicsPredictionEnabled())
+				{
+					const int32 LocalToServerOffset = PC->GetLocalToServerAsyncPhysicsTickOffset();
+					CachedServerFrame = PhysicsStep + LocalToServerOffset;
+				}
+				else if (ClientFrameInfo.LastProcessedInputFrame != INDEX_NONE)
 				{
 					const int32 FrameOffset = ClientFrameInfo.GetLocalFrameOffset();
-					
 					CachedServerFrame = PhysicsStep - FrameOffset; // Local = Server + Offset
 				}
 
@@ -278,7 +299,7 @@ struct FAsyncPhysicsInputRewindCallback : public Chaos::IRewindCallback
 						{
 							FObjectKey Key(RootComponent);
 							FRigidBodyState& LatestState = Scene->ReplicationCache.Map.FindOrAdd(Key);
-
+							
 							// This might be wrong... see FBodyInstance::GetRigidBodyState (converts to unreal units?) and FRepMovement::FillFrom
 							LatestState.Position = Handle->X();
 							LatestState.Quaternion = Handle->R();
@@ -328,7 +349,7 @@ void UAsyncPhysicsInputComponent::SetDataClass(TSubclassOf<UAsyncPhysicsData> In
 					// 1 frame is required to enable rewind capture. 
 					NumFrames = FMath::Max<int32>(1, NumFramesCVar->GetInt());
 				}
-				Solver->EnableRewindCapture(NumFrames, false, MakeUnique<FAsyncPhysicsInputRewindCallback>(World));
+				Solver->EnableRewindCapture(NumFrames, true, MakeUnique<FAsyncPhysicsInputRewindCallback>(World));
 			}
 
 			FAsyncPhysicsInputRewindCallback* Callback = static_cast<FAsyncPhysicsInputRewindCallback*>(Solver->GetRewindCallback());
diff --git a/Engine/Source/Runtime/Engine/Private/PhysicsEngine/BodyInstance.cpp b/Engine/Source/Runtime/Engine/Private/PhysicsEngine/BodyInstance.cpp
--- a/Engine/Source/Runtime/Engine/Private/PhysicsEngine/BodyInstance.cpp
+++ b/Engine/Source/Runtime/Engine/Private/PhysicsEngine/BodyInstance.cpp
@@ -3124,6 +3124,7 @@ float FBodyInstance::GetSleepThresholdMultiplier() const
 
 void FBodyInstance::SetLinearVelocity(const FVector& NewVel, bool bAddToCurrent, bool bAutoWake)
 {
+
 	FPhysicsCommand::ExecuteWrite(ActorHandle, [&](const FPhysicsActorHandle& Actor)
 	{
 		if(FPhysicsInterface::IsRigidBody(Actor))
@@ -3213,30 +3214,55 @@ void FBodyInstance::AddCustomPhysics(FCalculateCustomPhysics& CalculateCustomPhy
 	});
 }
 
-void FBodyInstance::AddForce(const FVector& Force, bool bAllowSubstepping, bool bAccelChange)
+void FBodyInstance::ApplyAsyncPhysicsCommand(FAsyncPhysicsTimestamp TimeStamp, const bool bIsInternal, APlayerController* PlayerController, const TFunction<void()>& Command)
 {
-	FPhysicsCommand::ExecuteWrite(ActorHandle, [&](const FPhysicsActorHandle& Actor)
+	if (bIsInternal && TimeStamp.IsValid())
 	{
-		if(!IsRigidBodyKinematic_AssumesLocked(Actor))
+		APlayerController* LocalController = PlayerController ? PlayerController : 
+			OwnerComponent->GetWorld() ? OwnerComponent->GetWorld()->GetFirstPlayerController() : nullptr;
+		if (LocalController)
 		{
-			if(FPhysScene* PhysScene = GetPhysicsScene())
+			TimeStamp.LocalFrame = TimeStamp.ServerFrame - LocalController->GetLocalToServerAsyncPhysicsTickOffset();
+			LocalController->ExecuteAsyncPhysicsCommand(TimeStamp, OwnerComponent.Get(), Command);
+		}
+	}
+	else
+	{
+		FPhysicsCommand::ExecuteWrite(ActorHandle, [&](const FPhysicsActorHandle& Actor)
 			{
-				PhysScene->AddForce_AssumesLocked(this, Force, bAllowSubstepping, bAccelChange);
-			}
+				Command();
+			});
+	}
+}
+
+static bool IsBodyDynamic(const FPhysicsActorHandle& ActorHandle, const bool bIsInternal)
+{
+	const Chaos::EObjectStateType ObjectState = bIsInternal ? ActorHandle->GetPhysicsThreadAPI()->ObjectState() :
+		ActorHandle->GetGameThreadAPI().ObjectState();
+
+	return (ObjectState == Chaos::EObjectStateType::Dynamic) || (ObjectState == Chaos::EObjectStateType::Sleeping);
+}
+
+void FBodyInstance::AddForce(const FVector& Force, bool bAllowSubstepping, bool bAccelChange, const FAsyncPhysicsTimestamp TimeStamp, APlayerController* PlayerController)
+{
+	const bool bIsInternal = TimeStamp.IsValid();
+	ApplyAsyncPhysicsCommand(TimeStamp, bIsInternal, PlayerController, [&, Force, bAllowSubstepping, bAccelChange, bIsInternal]()
+	{
+		if (FPhysicsInterface::IsInScene(ActorHandle) && IsBodyDynamic(ActorHandle, bIsInternal))
+		{
+			FPhysicsInterface::AddForce_AssumesLocked(ActorHandle, Force, bAllowSubstepping, bAccelChange, bIsInternal);
 		}
 	});
 }
 
-void FBodyInstance::AddForceAtPosition(const FVector& Force, const FVector& Position, bool bAllowSubstepping, bool bIsLocalForce)
+void FBodyInstance::AddForceAtPosition(const FVector& Force, const FVector& Position, bool bAllowSubstepping, bool bIsLocalForce, const FAsyncPhysicsTimestamp TimeStamp, APlayerController* PlayerController)
 {
-	FPhysicsCommand::ExecuteWrite(ActorHandle, [&](const FPhysicsActorHandle& Actor)
+	const bool bIsInternal = TimeStamp.IsValid();
+	ApplyAsyncPhysicsCommand(TimeStamp, bIsInternal, PlayerController, [&, Force, Position, bAllowSubstepping, bIsLocalForce, bIsInternal]()
 	{
-		if(!IsRigidBodyKinematic_AssumesLocked(Actor))
+		if (FPhysicsInterface::IsInScene(ActorHandle) && IsBodyDynamic(ActorHandle, bIsInternal))
 		{
-			if(FPhysScene* PhysScene = GetPhysicsScene())
-			{
-				PhysScene->AddForceAtPosition_AssumesLocked(this, Force, Position, bAllowSubstepping, bIsLocalForce);
-			}
+			FPhysicsInterface::AddForceAtPosition_AssumesLocked(ActorHandle, Force, Position, bAllowSubstepping, bIsLocalForce, bIsInternal);
 		}
 	});
 }
@@ -3266,16 +3292,14 @@ void FBodyInstance::SetOneWayInteraction(bool InOneWayInteraction /*= true*/)
 		});
 }
 
-void FBodyInstance::AddTorqueInRadians(const FVector& Torque, bool bAllowSubstepping, bool bAccelChange)
+void FBodyInstance::AddTorqueInRadians(const FVector& Torque, bool bAllowSubstepping, bool bAccelChange, const FAsyncPhysicsTimestamp TimeStamp, APlayerController* PlayerController)
 {
-	FPhysicsCommand::ExecuteWrite(ActorHandle, [&](const FPhysicsActorHandle& Actor)
+	const bool bIsInternal = TimeStamp.IsValid();
+	ApplyAsyncPhysicsCommand(TimeStamp, bIsInternal, PlayerController, [&, Torque, bAllowSubstepping, bAccelChange, bIsInternal]()
 	{
-		if(!IsRigidBodyKinematic_AssumesLocked(Actor))
+		if (FPhysicsInterface::IsInScene(ActorHandle) && IsBodyDynamic(ActorHandle, bIsInternal))
 		{
-			if(FPhysScene* PhysScene = GetPhysicsScene())
-			{
-				PhysScene->AddTorque_AssumesLocked(this, Torque, bAllowSubstepping, bAccelChange);
-			}
+			FPhysicsInterface::AddTorque_AssumesLocked(ActorHandle, Torque, bAllowSubstepping, bAccelChange, bIsInternal);
 		}
 	});
 }
@@ -3294,63 +3318,64 @@ void FBodyInstance::ClearTorques(bool bAllowSubstepping)
 	});
 }
 
-void FBodyInstance::AddAngularImpulseInRadians(const FVector& AngularImpulse, bool bVelChange)
+void FBodyInstance::AddAngularImpulseInRadians(const FVector& AngularImpulse, bool bVelChange, const FAsyncPhysicsTimestamp TimeStamp, APlayerController* PlayerController)
 {
-	FPhysicsCommand::ExecuteWrite(ActorHandle, [&](const FPhysicsActorHandle& Actor)
+	const bool bIsInternal = TimeStamp.IsValid();
+	ApplyAsyncPhysicsCommand(TimeStamp, bIsInternal, PlayerController, [&, AngularImpulse, bVelChange, bIsInternal]()
 	{
-		if (FPhysicsInterface::IsRigidBody(Actor) && FPhysicsInterface::IsInScene(Actor) && !IsRigidBodyKinematic_AssumesLocked(Actor))
+		if (FPhysicsInterface::IsInScene(ActorHandle) && IsBodyDynamic(ActorHandle, bIsInternal))
 		{
-			if (FPhysScene* PhysScene = GetPhysicsScene())
+			if (bVelChange)
 			{
-				if (bVelChange)
-				{
-					FPhysicsInterface::AddAngularVelocityInRadians_AssumesLocked(Actor, AngularImpulse);
-				}
-				else
-				{
-					FPhysicsInterface::AddAngularImpulseInRadians_AssumesLocked(Actor, AngularImpulse);
-				}
+				FPhysicsInterface::AddAngularVelocityInRadians_AssumesLocked(ActorHandle, AngularImpulse, bIsInternal);
+			}
+			else
+			{
+				FPhysicsInterface::AddAngularImpulseInRadians_AssumesLocked(ActorHandle, AngularImpulse, bIsInternal);
 			}
 		}
 	});
 }
 
-void FBodyInstance::AddImpulse(const FVector& Impulse, bool bVelChange)
+void FBodyInstance::AddImpulse(const FVector& Impulse, bool bVelChange, const FAsyncPhysicsTimestamp TimeStamp, APlayerController* PlayerController)
 {
-	FPhysicsCommand::ExecuteWrite(ActorHandle, [&](const FPhysicsActorHandle& Actor)
-	{
-		if(FPhysicsInterface::IsRigidBody(Actor) && FPhysicsInterface::IsInScene(Actor) && !IsRigidBodyKinematic_AssumesLocked(Actor))
+	const bool bIsInternal = TimeStamp.IsValid();
+	ApplyAsyncPhysicsCommand(TimeStamp, bIsInternal, PlayerController, [&, Impulse, bVelChange, bIsInternal]()
+	{		
+		if (FPhysicsInterface::IsInScene(ActorHandle) && IsBodyDynamic(ActorHandle, bIsInternal))
 		{
-			if(bVelChange)
+			if (bVelChange)
 			{
-				FPhysicsInterface::AddVelocity_AssumesLocked(Actor, Impulse);
+				FPhysicsInterface::AddVelocity_AssumesLocked(ActorHandle, Impulse, bIsInternal);
 			}
 			else
 			{
-				FPhysicsInterface::AddImpulse_AssumesLocked(Actor, Impulse);
+				FPhysicsInterface::AddImpulse_AssumesLocked(ActorHandle, Impulse, bIsInternal);
 			}
 		}
 	});
 }
 
-void FBodyInstance::AddImpulseAtPosition(const FVector& Impulse, const FVector& Position)
+void FBodyInstance::AddImpulseAtPosition(const FVector& Impulse, const FVector& Position, const FAsyncPhysicsTimestamp TimeStamp, APlayerController* PlayerController)
 {
-	FPhysicsCommand::ExecuteWrite(ActorHandle, [&](const FPhysicsActorHandle& Actor)
+	const bool bIsInternal = TimeStamp.IsValid();
+	ApplyAsyncPhysicsCommand(TimeStamp, bIsInternal, PlayerController, [&, Impulse, Position, bIsInternal]()
 	{
-		if(FPhysicsInterface::IsRigidBody(Actor) && FPhysicsInterface::IsInScene(Actor) && !IsRigidBodyKinematic_AssumesLocked(Actor))
+		if (FPhysicsInterface::IsInScene(ActorHandle) && IsBodyDynamic(ActorHandle, bIsInternal))
 		{
-			FPhysicsInterface::AddImpulseAtLocation_AssumesLocked(Actor, Impulse, Position);
+			FPhysicsInterface::AddImpulseAtLocation_AssumesLocked(ActorHandle, Impulse, Position, bIsInternal);
 		}
 	});
 }
 
-void FBodyInstance::AddVelocityChangeImpulseAtLocation(const FVector& Impulse, const FVector& Position)
+void FBodyInstance::AddVelocityChangeImpulseAtLocation(const FVector& Impulse, const FVector& Position, const FAsyncPhysicsTimestamp TimeStamp, APlayerController* PlayerController)
 {
-	FPhysicsCommand::ExecuteWrite(ActorHandle, [&](const FPhysicsActorHandle& Actor)
+	const bool bIsInternal = TimeStamp.IsValid();
+	ApplyAsyncPhysicsCommand(TimeStamp, bIsInternal, PlayerController, [&, Impulse, Position, bIsInternal]()
 	{
-		if (FPhysicsInterface::IsRigidBody(Actor) && FPhysicsInterface::IsInScene(Actor) && !IsRigidBodyKinematic_AssumesLocked(Actor))
+		if (FPhysicsInterface::IsInScene(ActorHandle) && IsBodyDynamic(ActorHandle, bIsInternal))
 		{
-			FPhysicsInterface::AddVelocityChangeImpulseAtLocation_AssumesLocked(Actor, Impulse, Position);
+			FPhysicsInterface::AddVelocityChangeImpulseAtLocation_AssumesLocked(ActorHandle, Impulse, Position, bIsInternal);
 		}
 	});
 }
diff --git a/Engine/Source/Runtime/Engine/Private/PhysicsEngine/Experimental/PhysScene_Chaos.cpp b/Engine/Source/Runtime/Engine/Private/PhysicsEngine/Experimental/PhysScene_Chaos.cpp
--- a/Engine/Source/Runtime/Engine/Private/PhysicsEngine/Experimental/PhysScene_Chaos.cpp
+++ b/Engine/Source/Runtime/Engine/Private/PhysicsEngine/Experimental/PhysScene_Chaos.cpp
@@ -61,6 +61,7 @@ struct FPendingAsyncPhysicsCommand
 	int32 PhysicsStep;
 	TWeakObjectPtr<UObject> OwningObject;
 	TFunction<void()> Command;
+	bool bEnableResim = true;
 };
 
 
@@ -79,16 +80,62 @@ public:
 	{
 		using namespace Chaos;
 
+		if (!PendingCommands.IsEmpty())
+		{
+			if (!PendingCommands[0].OwningObject.IsStale())
+			{
+				if(UWorld* World = PendingCommands[0].OwningObject->GetWorld())
+				{
+					if (FPhysScene* PhysScene = World->GetPhysicsScene())
+					{
+						PhysScene->GetPhysicsReplication()->SetResimFrame(INDEX_NONE);
+					}
+				}
+			}
+		}
+
+		static IConsoleVariable* ResimNetworkPhysicsCVar = IConsoleManager::Get().FindConsoleVariable(TEXT("np2.ResimNetworkPhysicsPrediction"));
+		const bool bEnableResim =  (ResimNetworkPhysicsCVar && ResimNetworkPhysicsCVar->GetInt() == 1);
+
+		const IConsoleVariable* NumFramesCVar = IConsoleManager::Get().FindConsoleVariable(TEXT("p.RewindCaptureNumFrames"));
+		const int32 NumFrames = NumFramesCVar ? NumFramesCVar->GetInt() : 0;
+
 		for (int32 Idx = 0; Idx < PendingCommands.Num(); ++Idx)
 		{
 			const int32 CurrentFrame = static_cast<FPBDRigidsSolver*>(GetSolver())->GetCurrentFrame();
 
 			FPendingAsyncPhysicsCommand* PendingCommand = &PendingCommands[Idx];
 			bool bRemove = PendingCommand->OwningObject.IsStale();
-			if (!bRemove && PendingCommands[Idx].PhysicsStep <= CurrentFrame)
+
+			if (bEnableResim && PendingCommand->bEnableResim && PendingCommands[Idx].PhysicsStep > (CurrentFrame - NumFrames))
 			{
-				PendingCommands[Idx].Command();
-				bRemove = true;
+				if (PendingCommands[Idx].PhysicsStep < CurrentFrame && PendingCommands[Idx].OwningObject.IsValid())
+				{
+					if(UWorld* World = PendingCommands[Idx].OwningObject->GetWorld())
+					{
+						if (FPhysScene* PhysScene = World->GetPhysicsScene())
+						{
+							int32 ResimFrame = PhysScene->GetPhysicsReplication()->GetResimFrame();
+							ResimFrame = (ResimFrame == INDEX_NONE) ? PendingCommands[Idx].PhysicsStep :
+								FMath::Min(ResimFrame, PendingCommands[Idx].PhysicsStep);
+
+							PhysScene->GetPhysicsReplication()->SetResimFrame(ResimFrame);
+						}
+					}
+				}
+				else if (!bRemove && PendingCommands[Idx].PhysicsStep == CurrentFrame)
+				{
+					PendingCommands[Idx].Command();
+					bRemove = true;
+				}
+			}
+			else
+			{
+				if (!bRemove && PendingCommands[Idx].PhysicsStep <= CurrentFrame)
+				{
+					PendingCommands[Idx].Command();
+					bRemove = true;
+				}
 			}
 
 			if (bRemove)
@@ -1096,109 +1143,17 @@ void FPhysScene_Chaos::AddCustomPhysics_AssumesLocked(FBodyInstance* BodyInstanc
 
 void FPhysScene_Chaos::AddForce_AssumesLocked(FBodyInstance* BodyInstance, const FVector& Force, bool bAllowSubstepping, bool bAccelChange)
 {
-	using namespace Chaos;
-
-	FPhysicsActorHandle& Handle = BodyInstance->GetPhysicsActorHandle();
-	if (FPhysicsInterface::IsValid(Handle))
-	{
-		Chaos::FRigidBodyHandle_External& Body_External = Handle->GetGameThreadAPI();
-		EObjectStateType ObjectState = Body_External.ObjectState();
-		Body_External.SetObjectState(EObjectStateType::Dynamic);
-
-		if (bAccelChange)
-		{
-			const Chaos::FReal Mass = Body_External.M();
-			const Chaos::FVec3 Acceleration = Force * Mass;
-			Body_External.AddForce(Acceleration);
-		}
-		else
-		{
-			Body_External.AddForce(Force);
-		}
-	}
+	FPhysicsInterface::AddForce_AssumesLocked(BodyInstance->GetPhysicsActorHandle(), Force, bAllowSubstepping, bAccelChange);
 }
 
 void FPhysScene_Chaos::AddForceAtPosition_AssumesLocked(FBodyInstance* BodyInstance, const FVector& Force, const FVector& Position, bool bAllowSubstepping, bool bIsLocalForce /*= false*/)
 {
-	using namespace Chaos;
-
-	FPhysicsActorHandle& Handle = BodyInstance->GetPhysicsActorHandle();
-	if (ensure(FPhysicsInterface::IsValid(Handle)))
-	{
-		Chaos::FRigidBodyHandle_External& Body_External = Handle->GetGameThreadAPI();
-		EObjectStateType ObjectState = Body_External.ObjectState();
-		const Chaos::FVec3 WorldCOM = FParticleUtilitiesGT::GetCoMWorldPosition(&Body_External);
-
-		Body_External.SetObjectState(EObjectStateType::Dynamic);
-
-		if (bIsLocalForce)
-		{
-			const Chaos::FRigidTransform3 CurrentTransform = FParticleUtilitiesGT::GetActorWorldTransform(&Body_External);
-			const Chaos::FVec3 WorldPosition = CurrentTransform.TransformPosition(Position);
-			const Chaos::FVec3 WorldForce = CurrentTransform.TransformVector(Force);
-			const Chaos::FVec3 WorldTorque = Chaos::FVec3::CrossProduct(WorldPosition - WorldCOM, WorldForce);
-			Body_External.AddForce(WorldForce);
-			Body_External.AddTorque(WorldTorque);
-		}
-		else
-		{
-			const Chaos::FVec3 WorldTorque = Chaos::FVec3::CrossProduct(Position - WorldCOM, Force);
-			Body_External.AddForce(Force);
-			Body_External.AddTorque(WorldTorque);
-		}
-	}
+	FPhysicsInterface::AddForceAtPosition_AssumesLocked(BodyInstance->GetPhysicsActorHandle(), Force, Position, bAllowSubstepping, bIsLocalForce);
 }
 
 void FPhysScene_Chaos::AddRadialForceToBody_AssumesLocked(FBodyInstance* BodyInstance, const FVector& Origin, const float Radius, const float Strength, const uint8 Falloff, bool bAccelChange, bool bAllowSubstepping)
 {
-	FPhysicsActorHandle& Handle = BodyInstance->GetPhysicsActorHandle();
-	if (ensure(FPhysicsInterface::IsValid(Handle)))
-	{
-		Chaos::FRigidBodyHandle_External& Body_External = Handle->GetGameThreadAPI();
-		Chaos::EObjectStateType ObjectState = Body_External.ObjectState();
-		if (CHAOS_ENSURE(ObjectState == Chaos::EObjectStateType::Dynamic || ObjectState == Chaos::EObjectStateType::Sleeping))
-		{
-			const Chaos::FVec3 WorldCOM = Chaos::FParticleUtilitiesGT::GetCoMWorldPosition(&Body_External);
-
-			Chaos::FVec3 Direction = WorldCOM - Origin;
-			const Chaos::FReal Distance = Direction.Size();
-			if (Distance > Radius)
-			{
-				return;
-			}
-
-			Body_External.SetObjectState(Chaos::EObjectStateType::Dynamic);
-
-			if (Distance < 1e-4)
-			{
-				Direction = Chaos::FVec3(1, 0, 0);
-			}
-			else
-			{
-				Direction = Direction.GetUnsafeNormal();
-			}
-			Chaos::FVec3 Force(0, 0, 0);
-			CHAOS_ENSURE(Falloff < RIF_MAX);
-			if (Falloff == ERadialImpulseFalloff::RIF_Constant)
-			{
-				Force = Strength * Direction;
-			}
-			if (Falloff == ERadialImpulseFalloff::RIF_Linear)
-			{
-				Force = (Radius - Distance) / Radius * Strength * Direction;
-			}
-			if (bAccelChange)
-			{
-				const Chaos::FReal Mass = Body_External.M();
-				const Chaos::FVec3 Acceleration = Force * Mass;
-				Body_External.AddForce(Acceleration);
-			}
-			else
-			{
-				Body_External.AddForce(Force);
-			}
-		}
-	}
+	FPhysicsInterface::AddRadialForce_AssumesLocked(BodyInstance->GetPhysicsActorHandle(), Origin, Radius, Strength, Falloff, bAccelChange, bAllowSubstepping);
 }
 
 void FPhysScene_Chaos::ClearForces_AssumesLocked(FBodyInstance* BodyInstance, bool bAllowSubstepping)
@@ -1212,25 +1167,7 @@ void FPhysScene_Chaos::ClearForces_AssumesLocked(FBodyInstance* BodyInstance, bo
 
 void FPhysScene_Chaos::AddTorque_AssumesLocked(FBodyInstance* BodyInstance, const FVector& Torque, bool bAllowSubstepping, bool bAccelChange)
 {
-	using namespace Chaos;
-
-	FPhysicsActorHandle& Handle = BodyInstance->GetPhysicsActorHandle();
-	if (ensure(FPhysicsInterface::IsValid(Handle)))
-	{
-		Chaos::FRigidBodyHandle_External& Body_External = Handle->GetGameThreadAPI();
-		EObjectStateType ObjectState = Body_External.ObjectState();
-		if (CHAOS_ENSURE(ObjectState == EObjectStateType::Dynamic || ObjectState == EObjectStateType::Sleeping))
-		{
-			if (bAccelChange)
-			{
-				Body_External.AddTorque(FParticleUtilitiesXR::GetWorldInertia(&Body_External) * Torque);
-			}
-			else
-			{
-				Body_External.AddTorque(Torque);
-			}
-		}
-	}
+	FPhysicsInterface::AddTorque_AssumesLocked(BodyInstance->GetPhysicsActorHandle(), Torque, bAllowSubstepping, bAccelChange);
 }
 
 void FPhysScene_Chaos::ClearTorques_AssumesLocked(FBodyInstance* BodyInstance, bool bAllowSubstepping)
@@ -1246,7 +1183,7 @@ void FPhysScene_Chaos::SetKinematicTarget_AssumesLocked(FBodyInstance* BodyInsta
 {
 	// #todo : Implement
 	//for now just pass it into actor directly
-	FPhysInterface_Chaos::SetKinematicTarget_AssumesLocked(BodyInstance->GetPhysicsActorHandle(), TargetTM);
+	FPhysicsInterface::SetKinematicTarget_AssumesLocked(BodyInstance->GetPhysicsActorHandle(), TargetTM);
 }
 
 bool FPhysScene_Chaos::GetKinematicTarget_AssumesLocked(const FBodyInstance* BodyInstance, FTransform& OutTM) const
@@ -2121,10 +2058,10 @@ void FPhysScene_Chaos::UnregisterAsyncPhysicsTickActor(AActor* Actor)
 	}
 }
 
-void FPhysScene_Chaos::EnqueueAsyncPhysicsCommand(int32 PhysicsStep, UObject* OwningObject, const TFunction<void()>& Command)
+void FPhysScene_Chaos::EnqueueAsyncPhysicsCommand(int32 PhysicsStep, UObject* OwningObject, const TFunction<void()>& Command, const bool bEnableResim)
 {
 	EnableAsyncPhysicsTickCallback();
-	AsyncPhysicsTickCallback->PendingCommands.Add({ PhysicsStep, TWeakObjectPtr<UObject>(OwningObject), Command });
+	AsyncPhysicsTickCallback->PendingCommands.Add({ PhysicsStep, TWeakObjectPtr<UObject>(OwningObject), Command, bEnableResim });
 }
 
 TSharedPtr<IPhysicsReplicationFactory> FPhysScene_Chaos::PhysicsReplicationFactory;
diff --git a/Engine/Source/Runtime/Engine/Private/PhysicsEngine/PhysicsReplication.cpp b/Engine/Source/Runtime/Engine/Private/PhysicsEngine/PhysicsReplication.cpp
--- a/Engine/Source/Runtime/Engine/Private/PhysicsEngine/PhysicsReplication.cpp
+++ b/Engine/Source/Runtime/Engine/Private/PhysicsEngine/PhysicsReplication.cpp
@@ -2,7 +2,7 @@
 
 /*=============================================================================
 	PhysicsReplication.cpp: Code for updating body instance physics state based on replication
-=============================================================================*/ 
+=============================================================================*/
 
 #include "PhysicsReplication.h"
 #include "Engine/World.h"
@@ -142,35 +142,137 @@ bool FPhysicsReplication::ApplyRigidBodyState(float DeltaSeconds, FBodyInstance*
 
 	Chaos::FPhysicsSolver* Solver = PhysScene->GetSolver();
 	Chaos::FRewindData* RewindData = Solver->GetRewindData();
+	const FPhysicsActorHandle& Proxy = BI->GetPhysicsActorHandle();
 
-	if (RewindData && LocalFrame > RewindData->GetEarliestFrame_Internal() && LocalFrame < RewindData->CurrentFrame())
+	if (Proxy && Chaos::FPhysicsSolverBase::IsNetworkPhysicsPredictionEnabled())
 	{
-		auto Proxy = BI->GetPhysicsActorHandle();
-		const auto P = RewindData->GetPastStateAtFrame(*Proxy->GetHandle_LowLevel(), LocalFrame);
-
-#if 0
-		// Debugging/test: compare and print out if locations differ
-		if (FVector::DistSquared(PhysicsTarget.TargetState.Position, P.X()) > 1.f)
+		if (RewindData && (LocalFrame <= RewindData->CurrentFrame()) && (LocalFrame >= RewindData->GetEarliestFrame_Internal()))
 		{
-			const int32 EarliestFrame = RewindData->GetEarliestFrame_Internal();
-			const int32 CurrentFrame = RewindData->CurrentFrame();
+			static constexpr Chaos::FFrameAndPhase::EParticleHistoryPhase RewindPhase = Chaos::FFrameAndPhase::EParticleHistoryPhase::PostPushData;
+
+			if (Chaos::FPhysicsSolverBase::CanResimNetworkPhysicsPrediction())
+			{
+				if (LocalFrame < RewindData->CurrentFrame() && !RewindData->IsResim())
+				{
+					if (UWorld* World = GetOwningWorld())
+					{
+						if (World->GetNetMode() == NM_Client)
+						{
+							if (APlayerController* PlayerController = World->GetFirstPlayerController())
+							{
+								FAsyncPhysicsTimestamp TimeStamp;
+								TimeStamp.LocalFrame = RewindData->CurrentFrame();
+								if (Chaos::FPhysicsSolverBase::CanDebugNetworkPhysicsPrediction())
+								{
+									UE_LOG(LogTemp, Log, TEXT("Appy Rigid body state at local frame %d with offset = %d"), LocalFrame, PlayerController->GetLocalToServerAsyncPhysicsTickOffset());
+								}
+
+								PlayerController->ExecuteAsyncPhysicsCommand(TimeStamp, BI->OwnerComponent.Get(), [this, Proxy, PhysicsTarget, RewindData, LocalFrame]
+								{
+									if (RewindData->IsResim())
+										return;
+
+									auto PastState = RewindData->GetPastStateAtFrame(*Proxy->GetHandle_LowLevel(), LocalFrame, RewindPhase);
+
+									static const float ErrorThreshold = Chaos::FPhysicsSolverBase::NetworkPhysicsErrorThreshold();
+
+									const float ErrorPosition = (PastState.X() - PhysicsTarget.TargetState.Position).Size();
+									float ColorLerp = (ErrorPosition >= ErrorThreshold) ? 1.0f : 0.0f;
+
+									if (Chaos::FPhysicsSolverBase::CanDebugNetworkPhysicsPrediction())
+									{
+										UE_LOG(LogTemp, Log, TEXT("Particle Position Error = %f | Resim Trigger = %f | Server Frame = %d | Client Frame = %d"), ErrorPosition, ColorLerp, PhysicsTarget.ServerFrame, LocalFrame);
+										UE_LOG(LogTemp, Log, TEXT("Particle Target Position = %s | Current Position = %s"), *PhysicsTarget.TargetState.Position.ToString(), *PastState.X().ToString());
+										UE_LOG(LogTemp, Log, TEXT("Particle Target Velocity = %s | Current Velocity = %s"), *PhysicsTarget.TargetState.LinVel.ToString(), *PastState.V().ToString());
+										UE_LOG(LogTemp, Log, TEXT("Particle Target Quaternion = %s | Current Quaternion = %s"), *PhysicsTarget.TargetState.Quaternion.ToString(), *PastState.R().ToString());
+										UE_LOG(LogTemp, Log, TEXT("Particle Target Omega = %s | Current Omega= %s"), *PhysicsTarget.TargetState.AngVel.ToString(), *PastState.W().ToString());
+
+										if (UWorld* OwningWorld = GetOwningWorld())
+										{
+											static constexpr float BoxSize = 5.0f;
+											const FColor DebugColor = FLinearColor::LerpUsingHSV(FLinearColor::Green, FLinearColor::Red, ColorLerp).ToFColor(false);
+
+											DrawDebugBox(OwningWorld, PhysicsTarget.TargetState.Position, FVector(BoxSize, BoxSize, BoxSize), PhysicsTarget.TargetState.Quaternion, FColor::Orange, true, CharacterMovementCVars::NetCorrectionLifetime, 0, 1.0f);
+											DrawDebugBox(OwningWorld, PastState.X(), FVector(6, 6, 6), PastState.R(), DebugColor, true, CharacterMovementCVars::NetCorrectionLifetime, 0, 1.0f);
+
+											DrawDebugDirectionalArrow(OwningWorld, PastState.X(), PhysicsTarget.TargetState.Position, 5.0f, FColor::Green, true, CharacterMovementCVars::NetCorrectionLifetime, 0, 0.5f);
+										}
+									}
+
+									const bool bShouldSleep = (PhysicsTarget.TargetState.Flags & ERigidBodyFlags::Sleeping) != 0;
+
+									RewindData->SetTargetStateAtFrame(*Proxy->GetHandle_LowLevel(), LocalFrame, RewindPhase,
+										PhysicsTarget.TargetState.Position, PhysicsTarget.TargetState.Quaternion, 
+										PhysicsTarget.TargetState.LinVel, PhysicsTarget.TargetState.AngVel, bShouldSleep);
+
+									if (ColorLerp == 1.0f)
+									{
+										if (Chaos::FPBDRigidsSolver* RigidSolver = Proxy->GetSolver<Chaos::FPBDRigidsSolver>())
+										{
+											RigidSolver->GetEvolution()->GetIslandManager().SetParticleResimFrame(Proxy->GetHandle_LowLevel(), LocalFrame);
+										}
+										ResimFrame = (ResimFrame == INDEX_NONE) ? LocalFrame : FMath::Min(ResimFrame, LocalFrame);
+									}
+								}, false);
+							}
+						}
+					}
+				}
+			}
+			else
+			{
+				static const float LerpFactor = Chaos::FPhysicsSolverBase::NetworkPhysicsInterpolationLerp();
+
+				FRigidBodyState CurrentState;
+				BI->GetRigidBodyState(CurrentState);
+
+				const float ExtrapolatedDeltaSeconds = (RewindData->CurrentFrame() - LocalFrame) * Solver->GetLastDt();
+
+				FVector LinDiff;
+				float LinDiffSize;
+				FVector AngDiffAxis;
+				float AngDiff;
+				float AngDiffSize;
+
+				ComputeDeltas(CurrentState.Position, CurrentState.Quaternion, PhysicsTarget.TargetState.Position, PhysicsTarget.TargetState.Quaternion, LinDiff, LinDiffSize, AngDiffAxis, AngDiff, AngDiffSize);
+
+				const FVector ExtrapolatedLinVel = FVector(PhysicsTarget.TargetState.LinVel + LinDiff / ExtrapolatedDeltaSeconds);
+				const FVector ExtrapolatedAngVel = FVector(PhysicsTarget.TargetState.AngVel + AngDiffAxis * AngDiff / ExtrapolatedDeltaSeconds);
+
+				const FVector TargetLinVel = FMath::Lerp(CurrentState.LinVel, ExtrapolatedLinVel, LerpFactor);
+				const FVector TargetAngVel = FMath::Lerp(CurrentState.AngVel, ExtrapolatedAngVel, LerpFactor);
+
+				FAsyncPhysicsDesiredState AsyncDesiredState;
+				AsyncDesiredState.WorldTM = FTransform(CurrentState.Quaternion, CurrentState.Position);
+				AsyncDesiredState.LinearVelocity = TargetLinVel;
+				AsyncDesiredState.AngularVelocity = TargetAngVel;
+				AsyncDesiredState.Proxy = static_cast<Chaos::FSingleParticlePhysicsProxy*>(Proxy);
+				AsyncDesiredState.ErrorCorrection = { 0.0f, 0.0f, 0.0f, 0.0f };
+				AsyncDesiredState.bShouldSleep = (PhysicsTarget.TargetState.Flags & ERigidBodyFlags::Sleeping) != 0;
 
-			UE_LOG(LogTemp, Warning, TEXT(""));
-			UE_LOG(LogTemp, Warning, TEXT("Location differs"));
-			UE_LOG(LogTemp, Warning, TEXT("   Replicated: %s"), *PhysicsTarget.TargetState.Position.ToString());
-			UE_LOG(LogTemp, Warning, TEXT("   Historic: %s"), *P.X().ToString());
+				CurAsyncData->Buffer.Add(AsyncDesiredState);
+				return false;
+			}
 		}
-#endif
+		else if(RewindData && LocalFrame > 0)
+		{
+			UE_LOG(LogPhysics, Warning, TEXT("FPhysicsReplication::ApplyRigidBodyState target frame (%d) out of rewind data bounds (%d,%d)"), LocalFrame,
+				RewindData->GetEarliestFrame_Internal(), RewindData->CurrentFrame());
+		}
+		return true;
+	}
+	else
+	{
+		// Call into the old ApplyRigidBodyState function for now,
+		// "new leash mode" should replace this.
+		// Note that old ApplyRigidBodyState is overridden in other projects, so consider backwards compat path
+		const float PingSecondsOneWay = RewindData == nullptr ? InPingSecondsOneWay : (Solver->GetLastDt() * (Solver->GetCurrentFrame() - LocalFrame)) * 0.5f;
+		return ApplyRigidBodyState(DeltaSeconds, BI, PhysicsTarget, ErrorCorrection, PingSecondsOneWay, nullptr);
 	}
-
-	// Call into the old ApplyRigidBodyState function for now,
-	// "new leash mode" should replace this.
-	// Note that old ApplyRigidBodyState is overridden in other projects, so consider backwards compat path
-	float PingSecondsOneWay = RewindData == nullptr ? InPingSecondsOneWay : (Solver->GetLastDt() * NumPredictedFrames) * 0.5f;
-	return ApplyRigidBodyState(DeltaSeconds, BI, PhysicsTarget, ErrorCorrection, PingSecondsOneWay);
 }
 
-bool FPhysicsReplication::ApplyRigidBodyState(float DeltaSeconds, FBodyInstance* BI, FReplicatedPhysicsTarget& PhysicsTarget, const FRigidBodyErrorCorrection& ErrorCorrection, const float PingSecondsOneWay, bool* bDidHardSnap)
+bool FPhysicsReplication::ApplyRigidBodyState(float DeltaSeconds, FBodyInstance* BI, FReplicatedPhysicsTarget& PhysicsTarget, const FRigidBodyErrorCorrection& ErrorCorrection,
+	const float PingSecondsOneWay, bool* bDidHardSnap)
 {
 	if (ShouldSkipPhysicsReplication())
 	{
@@ -334,9 +436,9 @@ bool FPhysicsReplication::ApplyRigidBodyState(float DeltaSeconds, FBodyInstance*
 		{
 			PhysicsTarget.AccumulatedErrorSeconds = FMath::Max(PhysicsTarget.AccumulatedErrorSeconds - DeltaSeconds, 0.0f);
 		}
-
+		
 		// Hard snap if error accumulation or linear error is big enough, and clear the error accumulator.
-		const bool bHardSnap =
+		const bool bHardSnap = 
 			LinDiffSize > MaxLinearHardSnapDistance ||
 			PhysicsTarget.AccumulatedErrorSeconds > ErrorAccumulationSeconds ||
 			CharacterMovementCVars::AlwaysHardSnap;
@@ -385,7 +487,7 @@ bool FPhysicsReplication::ApplyRigidBodyState(float DeltaSeconds, FBodyInstance*
 
 				const FVector NewPos = FMath::Lerp(FVector(CurrentState.Position), FVector(TargetPos), PositionLerp);
 				const FQuat NewAng = FQuat::Slerp(CurrentState.Quaternion, TargetQuat, AngleLerp);
-
+				
 				BI->SetBodyTransform(FTransform(NewAng, NewPos), ETeleportType::ResetPhysics);
 				BI->SetLinearVelocity(NewLinVel, false);
 				BI->SetAngularVelocityInRadians(FMath::DegreesToRadians(NewAngVel), false);
@@ -511,8 +613,7 @@ void FPhysicsReplication::OnTick(float DeltaSeconds, TMap<TWeakObjectPtr<UPrimit
 				Chaos::FPhysicsSolver* Solver = PhysScene->GetSolver();
 				check(Solver);
 
-				static IConsoleVariable* EnableNetworkPhysicsCVar = IConsoleManager::Get().FindConsoleVariable(TEXT("np2.EnableNetworkPhysicsPrediction"));
-				if (EnableNetworkPhysicsCVar && EnableNetworkPhysicsCVar->GetInt() == 1)
+				if (FPhysicsSolverBase::IsNetworkPhysicsPredictionEnabled())
 				{
 					LocalFrameOffset = PlayerController->GetLocalToServerAsyncPhysicsTickOffset();
 				}
@@ -523,19 +624,19 @@ void FPhysicsReplication::OnTick(float DeltaSeconds, TMap<TWeakObjectPtr<UPrimit
 	}
 
 	const FRigidBodyErrorCorrection& PhysicErrorCorrection = UPhysicsSettings::Get()->PhysicErrorCorrection;
-	if(AsyncCallback)
+	if (AsyncCallback)
 	{
 		PrepareAsyncData_External(PhysicErrorCorrection);
 	}
-
+	
 	// Get the ping between this PC & the server
 	const float LocalPing = GetLocalPing();
-
 	for (auto Itr = ComponentsToTargets.CreateIterator(); Itr; ++Itr)
 	{
 		bool bRemoveItr = false;
 		if (UPrimitiveComponent* PrimComp = Itr.Key().Get())
 		{
+
 			if (FBodyInstance* BI = PrimComp->GetBodyInstance(Itr.Value().BoneName))
 			{
 				FReplicatedPhysicsTarget& PhysicsTarget = Itr.Value();
@@ -559,14 +660,13 @@ void FPhysicsReplication::OnTick(float DeltaSeconds, TMap<TWeakObjectPtr<UPrimit
 
 						if (UpdatedState.Flags & ERigidBodyFlags::NeedsUpdate)
 						{
-							const int32 LocalFrame = PhysicsTarget.ServerFrame + LocalFrameOffset;
+							const int32 LocalFrame = PhysicsTarget.ServerFrame - LocalFrameOffset;
 							const bool bRestoredState = ApplyRigidBodyState(DeltaSeconds, BI, PhysicsTarget, PhysicErrorCorrection, PingSecondsOneWay, LocalFrame, NumPredictedFrames);
 							// Need to update the component to match new position.
 							if (PhysicsReplicationCVars::SkipSkeletalRepOptimization == 0 || Cast<USkeletalMeshComponent>(PrimComp) == nullptr)	//simulated skeletal mesh does its own polling of physics results so we don't need to call this as it'll happen at the end of the physics sim
 							{
 								PrimComp->SyncComponentToRBPhysics();
 							}
-
 							if (bRestoredState)
 							{
 								bRemoveItr = true;
@@ -598,7 +698,7 @@ void FPhysicsReplication::Tick(float DeltaSeconds)
 }
 
 FPhysicsReplication::FPhysicsReplication(FPhysScene* InPhysicsScene)
-: PhysScene(InPhysicsScene)
+	: PhysScene(InPhysicsScene)
 {
 	using namespace Chaos;
 	CurAsyncData = nullptr;
@@ -627,7 +727,7 @@ void FPhysicsReplication::PrepareAsyncData_External(const FRigidBodyErrorCorrect
 void FPhysicsReplication::ApplyAsyncDesiredState(const float DeltaSeconds, const FAsyncPhysicsRepCallbackData* AsyncData)
 {
 	using namespace Chaos;
-	if(AsyncData)
+	if (AsyncData)
 	{
 		for (const FAsyncPhysicsDesiredState& State : AsyncData->Buffer)
 		{
@@ -647,8 +747,7 @@ void FPhysicsReplication::ApplyAsyncDesiredState(const float DeltaSeconds, const
 			Chaos::FSingleParticlePhysicsProxy* Proxy = State.Proxy;
 			auto* Handle = Proxy->GetPhysicsThreadAPI();
 
-
-			if(Handle && Handle->CanTreatAsRigid())
+			if (Handle && Handle->CanTreatAsRigid())
 			{
 				const FVector TargetPos = State.WorldTM.GetLocation();
 				const FQuat TargetQuat = State.WorldTM.GetRotation();
@@ -672,7 +771,7 @@ void FPhysicsReplication::ApplyAsyncDesiredState(const float DeltaSeconds, const
 
 				const FVector NewPos = FMath::Lerp(FVector(CurrentState.Position), TargetPos, PositionLerp);
 				const FQuat NewAng = FQuat::Slerp(CurrentState.Quaternion, TargetQuat, AngleLerp);
-				
+
 				Handle->SetX(NewPos);
 				Handle->SetR(NewAng);
 				Handle->SetV(NewLinVel);
@@ -692,6 +791,16 @@ void FPhysicsReplication::ApplyAsyncDesiredState(const float DeltaSeconds, const
 	}
 }
 
+int32 FPhysicsReplication::GetResimFrame() const
+{
+	return ResimFrame;
+}
+
+void FPhysicsReplication::SetResimFrame(const int32 InResimFrame)
+{
+	ResimFrame = InResimFrame;
+}
+
 void FPhysicsReplication::SetReplicatedTarget(UPrimitiveComponent* Component, FName BoneName, const FRigidBodyState& ReplicatedTarget, int32 ServerFrame)
 {
 	if (UWorld* OwningWorld = GetOwningWorld())
@@ -713,7 +822,7 @@ void FPhysicsReplication::SetReplicatedTarget(UPrimitiveComponent* Component, FN
 		Target->TargetState = ReplicatedTarget;
 		Target->BoneName = BoneName;
 		Target->ArrivedTimeSeconds = OwningWorld->GetTimeSeconds();
-
+		
 		ensure(!Target->PrevPos.ContainsNaN());
 		ensure(!Target->PrevPosTarget.ContainsNaN());
 		ensure(!Target->TargetState.Position.ContainsNaN());
diff --git a/Engine/Source/Runtime/Engine/Private/PlayerController.cpp b/Engine/Source/Runtime/Engine/Private/PlayerController.cpp
--- a/Engine/Source/Runtime/Engine/Private/PlayerController.cpp
+++ b/Engine/Source/Runtime/Engine/Private/PlayerController.cpp
@@ -112,8 +112,22 @@ namespace NetworkPhysicsCvars
 
 	int32 EnableNetworkPhysicsPrediction = 0;
 	FAutoConsoleVariableRef CVarEnableNetworkPhysicsPrediction(TEXT("np2.EnableNetworkPhysicsPrediction"), EnableNetworkPhysicsPrediction, TEXT("Enables network physics prediction"));
-}
 
+	int32 DebugNetworkPhysicsPrediction = 0;
+	FAutoConsoleVariableRef CVarDebugNetworkPhysicsPrediction(TEXT("np2.DebugNetworkPhysicsPrediction"), DebugNetworkPhysicsPrediction, TEXT("Debugs network physics prediction"));
+
+	int32 ResimNetworkPhysicsPrediction = 1;
+	FAutoConsoleVariableRef CVarResimNetworkPhysicsPrediction(TEXT("np2.ResimNetworkPhysicsPrediction"), ResimNetworkPhysicsPrediction, TEXT("Resim network physics prediction"));
+
+	float NetworkPhysicsPredictionErrorThreshold = 10.0f;
+	FAutoConsoleVariableRef CVarNetworkPhysicsPredictionErrorThreshold(TEXT("np2.NetworkPhysicsPredictionErrorThreshold"), NetworkPhysicsPredictionErrorThreshold, TEXT("Position error threshold that will trigger the client correction based on the server prediction"));
+	
+	float NetworkPhysicsPredictionInterpLerp = 0.1f;
+	FAutoConsoleVariableRef CVarNetworkPhysicsPredictionInterpLerp(TEXT("np2.NetworkPhysicsPredictionInterpLerp"), NetworkPhysicsPredictionInterpLerp, TEXT("State lerp value in between the target state and the current one in case resim is disabled or if the pawn is not possessed (continuous correction)"));
+
+	float NetworkPhysicsPredictionResimLerp = 1.0f;
+	FAutoConsoleVariableRef CVarNetworkPhysicsPredictionResimLerp(TEXT("np2.NetworkPhysicsPredictionResimLerp"), NetworkPhysicsPredictionResimLerp, TEXT("State lerp value in between the target state and the matching one in the history when resim is enabled"));
+}
 
 const float RetryClientRestartThrottleTime = 0.5f;
 const float RetryServerAcknowledgeThrottleTime = 0.25f;
@@ -199,6 +213,7 @@ APlayerController::APlayerController(const FObjectInitializer& ObjectInitializer
 	if(NetworkPhysicsCvars::EnableNetworkPhysicsPrediction == 1)
 	{
 		bAsyncPhysicsTickEnabled = true;
+		AsyncPhysicsDataClass = UAsyncPhysicsData::StaticClass();
 	}
 }
 
@@ -6037,17 +6052,17 @@ UAsyncPhysicsData* APlayerController::GetAsyncPhysicsDataToWrite() const
 
 const UAsyncPhysicsData* APlayerController::GetAsyncPhysicsDataToConsume() const
 {
-	return AsyncPhysicsDataComponent ? AsyncPhysicsDataComponent->GetDataToConsume() : nullptr;
+	return AsyncPhysicsDataComponent ? AsyncPhysicsDataComponent->GetDataToConsume() : nullptr;  
 }
 
-void APlayerController::ExecuteAsyncPhysicsCommand(const FAsyncPhysicsTimestamp& AsyncPhysicsTimestamp, UObject* OwningObject, const TFunction<void()>& Command)
+void APlayerController::ExecuteAsyncPhysicsCommand(const FAsyncPhysicsTimestamp& AsyncPhysicsTimestamp, UObject* OwningObject, const TFunction<void()>& Command, const bool bEnableResim)
 {
 	if(UWorld* World = GetWorld())
 	{
 		if(FPhysScene* PhysScene = World->GetPhysicsScene())
 		{
 			const int32 PhysicsStep = IsLocalController() ? AsyncPhysicsTimestamp.LocalFrame : AsyncPhysicsTimestamp.ServerFrame;
-			PhysScene->EnqueueAsyncPhysicsCommand(PhysicsStep, OwningObject, Command);
+			PhysScene->EnqueueAsyncPhysicsCommand(PhysicsStep, OwningObject, Command, bEnableResim);
 		}
 	}
 }
diff --git a/Engine/Source/Runtime/Engine/Public/Physics/Experimental/PhysScene_Chaos.h b/Engine/Source/Runtime/Engine/Public/Physics/Experimental/PhysScene_Chaos.h
--- a/Engine/Source/Runtime/Engine/Public/Physics/Experimental/PhysScene_Chaos.h
+++ b/Engine/Source/Runtime/Engine/Public/Physics/Experimental/PhysScene_Chaos.h
@@ -120,7 +120,7 @@ public:
 	void RegisterAsyncPhysicsTickActor(AActor* Actor);
 	void UnregisterAsyncPhysicsTickActor(AActor* Actor);
 
-	void EnqueueAsyncPhysicsCommand(int32 PhysicsStep, UObject* OwningObject, const TFunction<void()>& Command);
+	void EnqueueAsyncPhysicsCommand(int32 PhysicsStep, UObject* OwningObject, const TFunction<void()>& Command, const bool bEnableResim = false);
 
 
 	/**
diff --git a/Engine/Source/Runtime/Engine/Public/PhysicsReplication.h b/Engine/Source/Runtime/Engine/Public/PhysicsReplication.h
--- a/Engine/Source/Runtime/Engine/Public/PhysicsReplication.h
+++ b/Engine/Source/Runtime/Engine/Public/PhysicsReplication.h
@@ -76,7 +76,6 @@ struct FAsyncPhysicsDesiredState
 	Chaos::FSingleParticlePhysicsProxy* Proxy;
 	TOptional<ErrorCorrectionData> ErrorCorrection;
 	bool bShouldSleep;
-	int32 ServerFrame;
 };
 
 struct FBodyInstance;
@@ -106,6 +105,12 @@ public:
 	/** Remove the replicated target*/
 	virtual void RemoveReplicatedTarget(UPrimitiveComponent* Component);
 
+	/** Get the resim frame (min server frame from the targets) */
+	int32 GetResimFrame() const;
+
+	/** Set the resim frame for replication */
+	void SetResimFrame(const int32 InResimFrame);
+
 protected:
 
 	/** Update the physics body state given a set of replicated targets */
@@ -139,4 +144,6 @@ private:
 	FAsyncPhysicsRepCallbackData* CurAsyncData;	//async data being written into before we push into callback
 	friend FPhysicsReplicationAsyncCallback;
 
+	int32 ResimFrame = INDEX_NONE;
+
 };
\ No newline at end of file
diff --git a/Engine/Source/Runtime/Experimental/Chaos/Private/Chaos/ChaosMarshallingManager.cpp b/Engine/Source/Runtime/Experimental/Chaos/Private/Chaos/ChaosMarshallingManager.cpp
--- a/Engine/Source/Runtime/Experimental/Chaos/Private/Chaos/ChaosMarshallingManager.cpp
+++ b/Engine/Source/Runtime/Experimental/Chaos/Private/Chaos/ChaosMarshallingManager.cpp
@@ -214,13 +214,7 @@ TArray<FPushPhysicsData*> FChaosMarshallingManager::StealHistory_Internal(int32
 	{
 		History.Add(HistoryQueue_Internal[Idx]);
 	}
-
-	for(int32 Idx = UseNumFrames; Idx < HistoryQueue_Internal.Num(); ++Idx)
-	{
-		FreeData_Internal(HistoryQueue_Internal[Idx]);
-	}
-
-	HistoryQueue_Internal.Reset();
+	HistoryQueue_Internal.RemoveAt(0, UseNumFrames);
 	return History;
 }
 
diff --git a/Engine/Source/Runtime/Experimental/Chaos/Private/Chaos/Collision/PBDCollisionConstraint.cpp b/Engine/Source/Runtime/Experimental/Chaos/Private/Chaos/Collision/PBDCollisionConstraint.cpp
--- a/Engine/Source/Runtime/Experimental/Chaos/Private/Chaos/Collision/PBDCollisionConstraint.cpp
+++ b/Engine/Source/Runtime/Experimental/Chaos/Private/Chaos/Collision/PBDCollisionConstraint.cpp
@@ -1036,7 +1036,8 @@ namespace Chaos
 	{
 		int32 MatchIndex = INDEX_NONE;
 
-		if (bChaos_Manifold_EnableFrictionRestore)
+		const bool bEnableNetworkPhysicsResim = FPhysicsSolverBase::IsNetworkPhysicsPredictionEnabled() && FPhysicsSolverBase::CanResimNetworkPhysicsPrediction();
+		if (bChaos_Manifold_EnableFrictionRestore && !bEnableNetworkPhysicsResim)
 		{
 			const FManifoldPoint& ManifoldPoint = ManifoldPoints[ManifoldPointIndex];
 			if (!ManifoldPoint.Flags.bDisabled)
diff --git a/Engine/Source/Runtime/Experimental/Chaos/Private/Chaos/CollisionOneShotManifolds.cpp b/Engine/Source/Runtime/Experimental/Chaos/Private/Chaos/CollisionOneShotManifolds.cpp
--- a/Engine/Source/Runtime/Experimental/Chaos/Private/Chaos/CollisionOneShotManifolds.cpp
+++ b/Engine/Source/Runtime/Experimental/Chaos/Private/Chaos/CollisionOneShotManifolds.cpp
@@ -769,6 +769,13 @@ namespace Chaos
 			const FReal Margin2 = Constraint.GetCollisionMargin1();
 			const FRigidTransform3 Convex2ToConvex1Transform = Convex2Transform.GetRelativeTransformNoScale(Convex1Transform);
 
+			const bool bEnableNetworkPhysicsResim = FPhysicsSolverBase::IsNetworkPhysicsPredictionEnabled() && FPhysicsSolverBase::CanResimNetworkPhysicsPrediction();
+			
+			if (bEnableNetworkPhysicsResim)
+			{
+				Constraint.GetGJKWarmStartData().Reset();
+			}
+
 			// Find the deepest penetration. This is used to determine the planes and points to use for the manifold
 			// MaxMarginDelta is an upper bound on the distance from the contact on the rounded core shape to the actual shape surface. 
 			FReal MaxMarginDelta = FReal(0);
diff --git a/Engine/Source/Runtime/Experimental/Chaos/Private/Chaos/Island/IslandGroupManager.cpp b/Engine/Source/Runtime/Experimental/Chaos/Private/Chaos/Island/IslandGroupManager.cpp
--- a/Engine/Source/Runtime/Experimental/Chaos/Private/Chaos/Island/IslandGroupManager.cpp
+++ b/Engine/Source/Runtime/Experimental/Chaos/Private/Chaos/Island/IslandGroupManager.cpp
@@ -110,7 +110,7 @@ namespace Chaos
 			}
 		}
 
-		int32 FPBDIslandGroupManager::BuildGroups()
+		int32 FPBDIslandGroupManager::BuildGroups(const bool bIsResimming)
 		{
 			// The set of islands with some work to do
 			// NOTE: We do not add islands with no constraints even if they have particles. There is no need since the Particles'
@@ -118,8 +118,18 @@ namespace Chaos
 			TArray<FPBDIsland*> Islands;
 			for (int32 IslandIndex = 0; IslandIndex < IslandManager.NumIslands(); ++IslandIndex)
 			{
-				FPBDIsland* Island = IslandManager.GetIsland(IslandIndex);
-				if (!Island->IsSleeping() && (Island->GetNumConstraints() > 0))
+				FPBDIsland* Island = IslandManager.GetIsland(IslandIndex); 
+
+#if !(UE_BUILD_SHIPPING || UE_BUILD_TEST)
+				if (Chaos::FPhysicsSolverBase::IsNetworkPhysicsPredictionEnabled() && Chaos::FPhysicsSolverBase::CanDebugNetworkPhysicsPrediction())
+				{
+					if(bIsResimming)
+					{ 
+						UE_LOG(LogChaos, Log, TEXT("Chaos Island[%d] needs resim = %d"), IslandIndex, IslandManager.IslandNeedsResim(IslandIndex));
+					}
+				}
+#endif
+				if (!Island->IsSleeping() && (Island->GetNumConstraints() > 0) && (!bIsResimming || (bIsResimming && IslandManager.IslandNeedsResim(IslandIndex))))
 				{
 					Islands.Add(Island);
 				}
diff --git a/Engine/Source/Runtime/Experimental/Chaos/Private/Chaos/Island/IslandManager.cpp b/Engine/Source/Runtime/Experimental/Chaos/Private/Chaos/Island/IslandManager.cpp
--- a/Engine/Source/Runtime/Experimental/Chaos/Private/Chaos/Island/IslandManager.cpp
+++ b/Engine/Source/Runtime/Experimental/Chaos/Private/Chaos/Island/IslandManager.cpp
@@ -604,6 +604,7 @@ void FPBDIslandManager::InitIslands()
 			LocalIsland++;
 
 			Island->ReserveParticles(GraphIsland.NumNodes);
+			Island->SetResimFrame(INDEX_NONE);
 
 			// Reset of the sleep counter if the island is :
 			// - Non persistent since we are starting incrementing the counter once the island is persistent and if values below the threshold
@@ -629,7 +630,7 @@ void FPBDIslandManager::InitIslands()
 	IslandIndexing.SetNum(LocalIsland, false);
 }
 
-void FPBDIslandManager::SyncIslands(FPBDRigidsSOAs& Particles)
+void FPBDIslandManager::SyncIslands(FPBDRigidsSOAs& Particles, const bool IsResimming)
 {
 	// Update of the sync and sleep state of all particles in all awake islands
 	for(auto& Island : Islands)
@@ -665,6 +666,13 @@ void FPBDIslandManager::GraphNodeLevelAssigned(FGraphNode& GraphNode)
 #endif
 
 		Island->AddParticle(GraphNode.NodeItem);
+
+		if(GraphNode.ResimFrame != INDEX_NONE)
+		{
+			const int32 ResimFrame = (Island->GetResimFrame() == INDEX_NONE) ? 
+					GraphNode.ResimFrame : FMath::Min(Island->GetResimFrame(), GraphNode.ResimFrame);
+			Island->SetResimFrame(ResimFrame);
+		}
 	}
 }
 
@@ -696,7 +704,7 @@ void FPBDIslandManager::GraphEdgeLevelAssigned(FGraphEdge& GraphEdge)
 	Island->AddConstraint(GraphEdge.EdgeItem, GraphEdge.LevelIndex, GraphEdge.ColorIndex, SubSortKey);
 }
 
-void FPBDIslandManager::UpdateIslands(FPBDRigidsSOAs& Particles)
+void FPBDIslandManager::UpdateIslands(FPBDRigidsSOAs& Particles, const bool IsResimming)
 {
 	// Assign all particles and constraints to islands, merging and splitting islands as required
 	IslandGraph->UpdateGraph();
@@ -722,7 +730,7 @@ void FPBDIslandManager::UpdateIslands(FPBDRigidsSOAs& Particles)
 	//IslandGraph->ComputeColors();
 
 	// Sync the state of all particles and constraints (primarily sleep state)
-	SyncIslands(Particles);
+	SyncIslands(Particles, IsResimming);
 
 	// Sort each island's constraints based on level, color, etc
 	// This is not required when we have a persistent graph because we already added constraints to islands in 
@@ -897,6 +905,31 @@ bool FPBDIslandManager::IslandNeedsResim(const int32 IslandIndex) const
 	return Islands[GetGraphIndex(IslandIndex)]->NeedsResim();
 }
 
+int32 FPBDIslandManager::GetParticleResimFrame(const FGeometryParticleHandle* ParticleHandle) const
+{
+	if (ParticleHandle && IslandGraph->GraphNodes.IsValidIndex(ParticleHandle->ConstraintGraphIndex()))
+	{
+		return IslandGraph->GraphNodes[ParticleHandle->ConstraintGraphIndex()].ResimFrame;
+	}
+	return INDEX_NONE;
+}
+
+void FPBDIslandManager::SetParticleResimFrame(FGeometryParticleHandle* ParticleHandle, const int32 ResimFrame)
+{
+	if (ParticleHandle && IslandGraph->GraphNodes.IsValidIndex(ParticleHandle->ConstraintGraphIndex()))
+	{
+		IslandGraph->GraphNodes[ParticleHandle->ConstraintGraphIndex()].ResimFrame = ResimFrame;
+	}
+}
+
+void FPBDIslandManager::ResetParticleResimFrame(const int32 ResetFrame)
+{
+	for(auto& GraphNode : IslandGraph->GraphNodes)
+	{
+		GraphNode.ResimFrame = ResetFrame;
+	}
+}
+
 ///////////////////////////////////////////////////////////////////////////////////////////////////
 ///////////////////////////////////////////////////////////////////////////////////////////////////
 ///////////////////////////////////////////////////////////////////////////////////////////////////
diff --git a/Engine/Source/Runtime/Experimental/Chaos/Private/Chaos/Island/SolverIsland.cpp b/Engine/Source/Runtime/Experimental/Chaos/Private/Chaos/Island/SolverIsland.cpp
--- a/Engine/Source/Runtime/Experimental/Chaos/Private/Chaos/Island/SolverIsland.cpp
+++ b/Engine/Source/Runtime/Experimental/Chaos/Private/Chaos/Island/SolverIsland.cpp
@@ -8,18 +8,18 @@ namespace Chaos
 	namespace Private
 	{
 	
-		FPBDIsland::FPBDIsland(const int32 MaxConstraintContainers)
-		: IslandIndex(INDEX_NONE)
-		, bIsSleeping(false)
-		, bNeedsResim(false)
-		, bIsPersistent(true)
-		, SleepCounter(0)
-		, IslandParticles()
-		, IslandConstraintsByType()
-		, NumConstraints(0)
-		{
-			IslandConstraintsByType.SetNum(MaxConstraintContainers);
-		}
+	FPBDIsland::FPBDIsland(const int32 MaxConstraintContainers)
+	: IslandIndex(INDEX_NONE)
+	, bIsSleeping(false)
+	, bNeedsResim(false)
+	, bIsPersistent(true)
+	, SleepCounter(0)
+	, IslandParticles()
+	, IslandConstraintsByType()
+	, NumConstraints(0)
+{
+	IslandConstraintsByType.SetNum(MaxConstraintContainers);
+}
 
 		void FPBDIsland::Reuse()
 		{
@@ -166,6 +166,24 @@ namespace Chaos
 					break;
 				}
 			}
+			for (TArray<FPBDIslandConstraint>& TypedIslandConstraints : IslandConstraintsByType)
+			{
+				for (FPBDIslandConstraint& IslandConstraint : TypedIslandConstraints)
+				{
+					const TVec2<FGeometryParticleHandle*>& ConstraintParticles = IslandConstraint.GetConstraint()->GetConstrainedParticles();
+
+					if(ConstraintParticles[0] && (ConstraintParticles[0]->SyncState() != ESyncState::InSync))
+					{
+						bNeedsResim = true;
+						break;
+					}
+					if (ConstraintParticles[1] && (ConstraintParticles[1]->SyncState() != ESyncState::InSync))
+					{
+						bNeedsResim = true;
+						break;
+					}
+				}
+			}
 		}
 
 	}	// namespace Private
diff --git a/Engine/Source/Runtime/Experimental/Chaos/Private/Chaos/PBDRigidsEvolutionGBF.cpp b/Engine/Source/Runtime/Experimental/Chaos/Private/Chaos/PBDRigidsEvolutionGBF.cpp
--- a/Engine/Source/Runtime/Experimental/Chaos/Private/Chaos/PBDRigidsEvolutionGBF.cpp
+++ b/Engine/Source/Runtime/Experimental/Chaos/Private/Chaos/PBDRigidsEvolutionGBF.cpp
@@ -255,6 +255,7 @@ void FPBDRigidsEvolutionGBF::AdvanceOneTimeStep(const FReal Dt,const FSubStepInf
 
 void FPBDRigidsEvolutionGBF::ReloadParticlesCache()
 {
+	
 	// @todo(chaos): Parallelize
 	FEvolutionResimCache* ResimCache = GetCurrentStepResimCache();
 	if ((ResimCache != nullptr) && ResimCache->IsResimming())
@@ -265,6 +266,12 @@ void FPBDRigidsEvolutionGBF::ReloadParticlesCache()
 			bool bIsUsingCache = false;
 			if (!Island->IsSleeping() && !GetIslandManager().IslandNeedsResim(IslandIndex))
 			{
+#if !(UE_BUILD_SHIPPING || UE_BUILD_TEST)
+				if (Chaos::FPhysicsSolverBase::IsNetworkPhysicsPredictionEnabled() && Chaos::FPhysicsSolverBase::CanDebugNetworkPhysicsPrediction())
+				{
+					UE_LOG(LogChaos, Log, TEXT("Reloading Island[%d] cache for %d particles"), IslandIndex, Island->GetParticles().Num());
+				}
+#endif
 				for (Private::FPBDIslandParticle& IslandParticle : Island->GetParticles())
 				{
 					if (auto Rigid = IslandParticle.GetParticle()->CastToRigidParticle())
@@ -500,7 +507,13 @@ void FPBDRigidsEvolutionGBF::AdvanceOneTimeStepImpl(const FReal Dt, const FSubSt
 	{
 		SCOPE_CYCLE_COUNTER(STAT_Evolution_BuildGroups);
 		CSV_SCOPED_TIMING_STAT(PhysicsVerbose, StepSolver_BuildGroups);
-		NumGroups = IslandGroupManager.BuildGroups();
+
+		// If we are resimming and if we have a particle cache we only build the island the groups 
+		// for islands that required to be simulated based of desynced particles
+		FEvolutionResimCache* ResimCache = GetCurrentStepResimCache();
+		const bool bIsResimming = (ResimCache != nullptr) && ResimCache->IsResimming();
+
+		NumGroups = IslandGroupManager.BuildGroups(bIsResimming);
 	}
 
 
@@ -693,6 +706,14 @@ void FPBDRigidsEvolutionGBF::TestModeResetCollisions()
 	}
 }
 
+void FPBDRigidsEvolutionGBF::ResetCollisions()
+{
+	for (FPBDCollisionConstraintHandle* Collision : CollisionConstraints.GetConstraintHandles())
+	{
+		Collision->GetContact().GetGJKWarmStartData().Reset();
+	}
+}
+
 FPBDRigidsEvolutionGBF::FPBDRigidsEvolutionGBF(
 	FPBDRigidsSOAs& InParticles, 
 	THandleArray<FChaosPhysicsMaterial>& SolverPhysicsMaterials, 
diff --git a/Engine/Source/Runtime/Experimental/Chaos/Private/PBDRigidsSolver.cpp b/Engine/Source/Runtime/Experimental/Chaos/Private/PBDRigidsSolver.cpp
--- a/Engine/Source/Runtime/Experimental/Chaos/Private/PBDRigidsSolver.cpp
+++ b/Engine/Source/Runtime/Experimental/Chaos/Private/PBDRigidsSolver.cpp
@@ -328,6 +328,11 @@ namespace Chaos
 			LLM_SCOPE(ELLMTag::ChaosUpdate);
 			UE_LOG(LogPBDRigidsSolver, Verbose, TEXT("AdvanceOneTimeStepTask::DoWork()"));
 
+			if (FRewindData* RewindData = MSolver->GetRewindData())
+			{
+				RewindData->ApplyInputs(MSolver->GetCurrentFrame(), MSolver->GetEvolution()->IsResetting());
+			}
+
 			MSolver->ApplyCallbacks_Internal();
 			
 			{
@@ -843,7 +848,14 @@ namespace Chaos
 	void FPBDRigidsSolver::EnableRewindCapture(int32 NumFrames, bool InUseCollisionResimCache, TUniquePtr<IRewindCallback>&& RewindCallback)
 	{
 		//TODO: this function calls both internal and extrnal - sort of assumed during initialization. Should decide what thread it's called on and mark it as either external or internal
-		MRewindData = MakeUnique<FRewindData>(((FPBDRigidsSolver*)this), NumFrames, InUseCollisionResimCache, ((FPBDRigidsSolver*)this)->GetCurrentFrame()); // FIXME
+		if (MRewindData.IsValid())
+		{
+			MRewindData->Init(((FPBDRigidsSolver*)this), NumFrames, InUseCollisionResimCache, ((FPBDRigidsSolver*)this)->GetCurrentFrame());
+		}
+		else
+		{
+			MRewindData = MakeUnique<FRewindData>(((FPBDRigidsSolver*)this), NumFrames, InUseCollisionResimCache, ((FPBDRigidsSolver*)this)->GetCurrentFrame()); // FIXME
+		}
 		bUseCollisionResimCache = InUseCollisionResimCache;
 		MRewindCallback = MoveTemp(RewindCallback);
 		MarshallingManager.SetHistoryLength_Internal(NumFrames);
@@ -1052,6 +1064,13 @@ namespace Chaos
 	{
 		const FReal StartSimTime = GetSolverTime();
 
+#if !(UE_BUILD_SHIPPING || UE_BUILD_TEST)
+		if(IsNetworkPhysicsPredictionEnabled() && CanDebugNetworkPhysicsPrediction())
+		{
+			UE_LOG(LogChaos, Log, TEXT("-> Simulating Frame = %d"), CurrentFrame);
+		}
+#endif
+
 		AdvanceOneTimeStepTask(this, MLastDt, SubStepInfo).DoWork();
 
 		if (MLastDt > 0)
@@ -1524,15 +1543,39 @@ namespace Chaos
 		{
 			const int32 LastStep = MRewindData->CurrentFrame() - 1;
 			const int32 ResimStep = MRewindCallback->TriggerRewindIfNeeded_Internal(LastStep);
-			if(ResimStep != INDEX_NONE)
+			const int32 NumResimSteps = LastStep - ResimStep + 1;
+
+			const bool bEnableNetworlPredictionDebug = IsNetworkPhysicsPredictionEnabled() && CanDebugNetworkPhysicsPrediction();
+			
+#if !(UE_BUILD_SHIPPING || UE_BUILD_TEST)
+			if (bEnableNetworlPredictionDebug)
+			{
+				UE_LOG(LogChaos, Log, TEXT("	-> Trying to Rewind Frame = %d | At Time = %f | Num Steps = %d | Resim Step = %d | Last Step = %d"), CurrentFrame,  MTime, NumResimSteps, ResimStep, LastStep);
+				for(auto& Handle : GetParticles().GetNonDisabledDynamicView())
+				{
+					if(!Handle.IsSleeping())
+					{
+						UE_LOG(LogChaos, Log, TEXT("Particle Active At Position = %s | Velocity = %s | Quaternion = %s | Omega = %s"),
+							*Handle.X().ToString(), *Handle.V().ToString(), *Handle.R().ToString(), *Handle.W().ToString());
+					}
+					else
+					{
+						UE_LOG(LogChaos, Log, TEXT("Particle Sleeping At Position = %s | Velocity = %s | Quaternion = %s | Omega = %s"), 
+							*Handle.X().ToString(), *Handle.V().ToString(), *Handle.R().ToString(), *Handle.W().ToString());
+					}
+				}
+			}
+#endif
+			
+			if(ResimStep >= 0 && (ResimStep < LastStep) && NumResimSteps <= MarshallingManager.GetNumHistory_Internal())
 			{
 				FResimDebugInfo DebugInfo;
 				QUICK_SCOPE_CYCLE_COUNTER(ChaosRewindAndResim);
-				if(ensure(MRewindData->RewindToFrame(ResimStep)))
+				if(MRewindData->RewindToFrame(ResimStep))
 				{
 					GetEvolution()->SetResim(true);
 					CurrentFrame = ResimStep;
-					const int32 NumResimSteps = LastStep - ResimStep + 1;
+					
 					TArray<FPushPhysicsData*> RecordedPushData = MarshallingManager.StealHistory_Internal(NumResimSteps);
 					bool bFirst = true;
 
@@ -1540,28 +1583,54 @@ namespace Chaos
 					// Do rollback as necessary
 					for (int32 Step = ResimStep; Step <= LastStep; ++Step)
 					{
-						FPushPhysicsData* PushData = RecordedPushData[LastStep - Step];	//push data is sorted as latest first
-						if(bFirst)
+						if ((LastStep - Step) < RecordedPushData.Num())
 						{
-							MTime = PushData->StartTime;	//not sure if sub-steps have proper StartTime so just do this once and let solver evolve remaining time
-						}
-
-						MRewindCallback->PreResimStep_Internal(Step, bFirst);
-						FAllSolverTasks ImmediateTask(*this, PushData);
-						ensure(bSolverHasFrozenGameThreadCallbacks == false);	//We don't support this for resim as it's very expensive and difficult to schedule
-						ImmediateTask.AdvanceSolver();
-						MRewindCallback->PostResimStep_Internal(Step);
+							// We need to reset all the particles having received a target from the server even if not at first resim frame
+							MRewindData->ApplyTargets(Step, bFirst);
 
-						bFirst = false;
+							FPushPhysicsData* PushData = RecordedPushData[LastStep - Step];	//push data is sorted as latest first
+							if (bFirst)
+							{
+								MTime = PushData->StartTime;	//not sure if sub-steps have proper StartTime so just do this once and let solver evolve remaining time
+								GetEvolution()->ResetCollisions();
+							}
+								
+							GetEvolution()->SetReset(bFirst);
+#if !(UE_BUILD_SHIPPING || UE_BUILD_TEST)
+							if(bEnableNetworlPredictionDebug)
+							{
+								UE_LOG(LogChaos, Log, TEXT("		-> Re-simulating Frame = %d "), Step);
+								for(auto& Handle : GetParticles().GetNonDisabledDynamicView())
+								{
+									if(!Handle.IsSleeping()) 
+									{
+										UE_LOG(LogChaos, Log, TEXT("Particle Dynamic At Position = %s | Velocity = %s | Quaternion = %s | Omega = %s | Resim Frame = %d | Sync State = %d | Needs Resim = %d"),   
+											*Handle.X().ToString(), *Handle.V().ToString(), *Handle.R().ToString(), *Handle.W().ToString(), 
+												GetEvolution()->GetIslandManager().GetIsland(Handle.IslandIndex())->GetResimFrame(), (uint8)Handle.SyncState(),
+											GetEvolution()->GetIslandManager().GetIsland(Handle.IslandIndex())->NeedsResim());
+									}
+								}
+							}
+#endif
+							MRewindCallback->PreResimStep_Internal(Step, bFirst);
+							
+							FAllSolverTasks ImmediateTask(*this, PushData);
+							//ensure(bSolverHasFrozenGameThreadCallbacks == false);	//We don't support this for resim as it's very expensive and difficult to schedule
+							ImmediateTask.AdvanceSolver();
+							MRewindCallback->PostResimStep_Internal(Step);
+							
+							bFirst = false;
+						}
 					}
+					GetEvolution()->GetIslandManager().ResetParticleResimFrame();
 
 					GetEvolution()->SetResim(false);
+					GetEvolution()->SetReset(false);
 
 					ResimTimer.Stop();
 					MRewindCallback->SetResimDebugInfo_Internal(DebugInfo);
 				}
 			}
-
 		}
 	}
 
@@ -2058,7 +2127,7 @@ CSV_CUSTOM_STAT(PhysicsCounters, Name, Value, ECsvCustomStatOp::Set);
 
 	bool FPBDRigidsSolver::IsDetemerministic() const
 	{
-		return bIsDeterministic || (MRewindData != nullptr) || (ChaosSolverDeterministic >= 1);
+		return bIsDeterministic || (MRewindData != nullptr) || (ChaosSolverDeterministic >= 1) || FPhysicsSolverBase::IsNetworkPhysicsPredictionEnabled();
 	}
 
 	void FPBDRigidsSolver::SetIsDeterministic(const bool bInIsDeterministic)
diff --git a/Engine/Source/Runtime/Experimental/Chaos/Private/RewindData.cpp b/Engine/Source/Runtime/Experimental/Chaos/Private/RewindData.cpp
--- a/Engine/Source/Runtime/Experimental/Chaos/Private/RewindData.cpp
+++ b/Engine/Source/Runtime/Experimental/Chaos/Private/RewindData.cpp
@@ -70,52 +70,76 @@ bool SimWritablePropsMayChange(const TGeometryParticleHandle<FReal,3>& Handle)
 	return ObjectState == EObjectStateType::Dynamic || ObjectState == EObjectStateType::Sleeping;
 }
 
-template <bool bSkipDynamics>
-bool FGeometryParticleStateBase::IsInSync(const FGeometryParticleHandle& Handle, const FFrameAndPhase FrameAndPhase, const FDirtyPropertiesPool& Pool) const
+bool FGeometryParticleStateBase::IsResimFrameValid(const FGeometryParticleHandle& Handle, const FFrameAndPhase FrameAndPhase) const
 {
-	if(!ParticlePositionRotation.IsInSync(Handle, FrameAndPhase, Pool))
+	if (FPBDRigidsSolver* RigidSolver = Handle.PhysicsProxy()->GetSolver<FPBDRigidsSolver>())
 	{
-		return false;
-	}
+		if (RigidSolver->GetEvolution() && RigidSolver->GetEvolution()->IsResimming())
+		{
+			const int32 ResimFrame = RigidSolver->GetEvolution()->GetIslandManager().GetParticleResimFrame(&Handle);
 
-	if(!NonFrequentData.IsInSync(Handle, FrameAndPhase, Pool))
-	{
-		return false;
+			if (ResimFrame != INDEX_NONE && ResimFrame <= FrameAndPhase.Frame)
+			{
+				return false;
+			}
+		}
 	}
+	return true;
+}
 
-	//todo: deal with state change mismatch
-
-	if(auto Kinematic = Handle.CastToKinematicParticle())
-	{
-		if(!Velocities.IsInSync(*Kinematic, FrameAndPhase, Pool))
+template <bool bSkipDynamics>
+bool FGeometryParticleStateBase::IsInSync(const FGeometryParticleHandle& Handle, const FFrameAndPhase FrameAndPhase, const FDirtyPropertiesPool& Pool) const
+{
+	if(Chaos::FPhysicsSolverBase::IsNetworkPhysicsPredictionEnabled() && !TargetPositions.IsEmpty() && !TargetVelocities.IsEmpty())
+	{ 
+		return IsResimFrameValid(Handle, FrameAndPhase);
+	}
+	else
+	{ 
+		if(!ParticlePositionRotation.IsInSync(Handle, FrameAndPhase, Pool))
 		{
 			return false;
 		}
 
-		if (!KinematicTarget.IsInSync(*Kinematic, FrameAndPhase, Pool))
+		if(!NonFrequentData.IsInSync(Handle, FrameAndPhase, Pool))
 		{
 			return false;
 		}
-	}
 
-	if(auto Rigid = Handle.CastToRigidParticle())
-	{
-		if(!bSkipDynamics)
+		//todo: deal with state change mismatch
+
+		if(auto Kinematic = Handle.CastToKinematicParticle())
 		{
-			if (!Dynamics.IsInSync(*Rigid, FrameAndPhase, Pool))
+			if(!Velocities.IsInSync(*Kinematic, FrameAndPhase, Pool))
 			{
 				return false;
 			}
-		}
 
-		if(!DynamicsMisc.IsInSync(*Rigid, FrameAndPhase, Pool))
-		{
-			return false;
+			if (!KinematicTarget.IsInSync(*Kinematic, FrameAndPhase, Pool))
+			{
+				return false;
+			}
 		}
 
-		if(!MassProps.IsInSync(*Rigid, FrameAndPhase, Pool))
+		if(auto Rigid = Handle.CastToRigidParticle())
 		{
-			return false;
+			if(!bSkipDynamics)
+			{
+				if (!Dynamics.IsInSync(*Rigid, FrameAndPhase, Pool))
+				{
+					return false;
+				}
+			}
+
+			if(!DynamicsMisc.IsInSync(*Rigid, FrameAndPhase, Pool))
+			{
+				return false;
+			}
+
+			if(!MassProps.IsInSync(*Rigid, FrameAndPhase, Pool))
+			{
+				return false;
+			}
 		}
 	}
 
@@ -139,6 +163,60 @@ bool FJointStateBase::IsInSync(const FPBDJointConstraintHandle& Handle, const FF
 	return true;
 }
 
+void FRewindData::ApplyInputs(const int32 ApplyFrame, const bool bResetSolver)
+{
+	for (TWeakPtr<FBaseInputsHistory>& InputsHistory : InputsHistories)
+	{
+		if (InputsHistory.IsValid())
+		{
+			InputsHistory.Pin().Get()->ApplyInputs(ApplyFrame, bResetSolver);
+		}
+	}
+}
+void FRewindData::RewindStates(const int32 RewindFrame, const bool bResetSolver)
+{
+	for (TWeakPtr<FBaseStatesHistory>& StatesHistory : StatesHistories)
+	{
+		if (StatesHistory.IsValid())
+		{
+			StatesHistory.Pin().Get()->RewindStates(RewindFrame, bResetSolver);
+		}
+	}
+}
+
+void FRewindData::ApplyTargets(const int32 Frame, const bool bResetSimulation)
+{
+	RewindStates(Frame, bResetSimulation);
+
+	EnsureIsInPhysicsThreadContext();
+
+	//If property changed between Frame and CurFrame, record the latest value and rewind to old
+	FFrameAndPhase RewindFrameAndPhase{ Frame, FFrameAndPhase::PostPushData };
+
+	auto RewindHelper = [RewindFrameAndPhase, this](auto Obj, bool bResimAsFollower, auto& Property, const auto& RewindFunc)
+	{
+		if (!Property.IsClean(RewindFrameAndPhase) && !bResimAsFollower)
+		{
+			RewindFunc(Obj, *Property.Read(RewindFrameAndPhase, PropertiesPool));
+		}
+	};
+
+	for (FDirtyParticleInfo& DirtyParticleInfo : DirtyParticles)
+	{
+		FGeometryParticleHandle* PTParticle = DirtyParticleInfo.GetObjectPtr();
+		FGeometryParticleStateBase& History = DirtyParticleInfo.GetHistory();	
+
+		const bool bResimAsFollower = DirtyParticleInfo.bResimAsFollower;
+		RewindHelper(PTParticle, bResimAsFollower, History.TargetPositions, [](auto Particle, const auto& Data) {Particle->SetXR(Data); });
+		RewindHelper(PTParticle->CastToKinematicParticle(), bResimAsFollower, History.TargetVelocities, [](auto Particle, const auto& Data) {Particle->SetVelocities(Data); });
+
+		if (!History.TargetPositions.IsClean(RewindFrameAndPhase) && Chaos::FPhysicsSolverBase::CanDebugNetworkPhysicsPrediction())
+		{
+			UE_LOG(LogChaos, Log, TEXT("Reset particle %d position to the target %s at frame %d"), PTParticle->UniqueIdx().Idx, *PTParticle->X().ToString(), Frame);
+		}
+	}
+}
+
 bool FRewindData::RewindToFrame(int32 Frame)
 {
 	QUICK_SCOPE_CYCLE_COUNTER(RewindToFrame);
@@ -196,13 +274,14 @@ bool FRewindData::RewindToFrame(int32 Frame)
 		FGeometryParticleStateBase& History = DirtyParticleInfo.AddFrame(CurFrame);	//non-const in case we need to record what's at head for a rewind (CurFrame has already been increased to the next frame)
 
 		const bool bResimAsFollower = DirtyParticleInfo.bResimAsFollower;
+
 		bool bAnyChange = RewindHelper(PTParticle, bResimAsFollower, History.ParticlePositionRotation, [](auto Particle, const auto& Data) {Particle->SetXR(Data); });
-		bAnyChange |= RewindHelper(PTParticle, bResimAsFollower, History.NonFrequentData, [](auto Particle, const auto& Data) {Particle->SetNonFrequentData(Data); });
 		bAnyChange |= RewindHelper(PTParticle->CastToKinematicParticle(), bResimAsFollower, History.Velocities, [](auto Particle, const auto& Data) {Particle->SetVelocities(Data); });
+		bAnyChange |= RewindHelper(PTParticle, bResimAsFollower, History.NonFrequentData, [](auto Particle, const auto& Data) {Particle->SetNonFrequentData(Data); });
 		bAnyChange |= RewindHelper(PTParticle->CastToKinematicParticle(), bResimAsFollower, History.KinematicTarget, [](auto Particle, const auto& Data) {Particle->SetKinematicTarget(Data); });
-		bAnyChange |= RewindHelper(PTParticle->CastToRigidParticle(),bResimAsFollower,  History.Dynamics, [](auto Particle, const auto& Data) {Particle->SetDynamics(Data); });
-		bAnyChange |= RewindHelper(PTParticle->CastToRigidParticle(),bResimAsFollower,  History.DynamicsMisc, [this](auto Particle, const auto& Data) {Solver->SetParticleDynamicMisc(Particle, Data); });
-		bAnyChange |= RewindHelper(PTParticle->CastToRigidParticle(),bResimAsFollower,  History.MassProps, [](auto Particle, const auto& Data) {Particle->SetMassProps(Data); });
+		bAnyChange |= RewindHelper(PTParticle->CastToRigidParticle(), bResimAsFollower, History.Dynamics, [](auto Particle, const auto& Data) {Particle->SetDynamics(Data); });
+		bAnyChange |= RewindHelper(PTParticle->CastToRigidParticle(), bResimAsFollower, History.DynamicsMisc, [this](auto Particle, const auto& Data) {Solver->SetParticleDynamicMisc(Particle, Data); });
+		bAnyChange |= RewindHelper(PTParticle->CastToRigidParticle(), bResimAsFollower, History.MassProps, [](auto Particle, const auto& Data) {Particle->SetMassProps(Data); });
 
 		if (!bResimAsFollower)
 		{
@@ -254,6 +333,7 @@ void FRewindData::DesyncIfNecessary(TDirtyInfo& Info, const FFrameAndPhase Frame
 
 	auto Handle = Info.GetObjectPtr();
 	const auto& History = Info.GetHistory();
+
 	if (Handle->SyncState() == ESyncState::InSync && !History.template IsInSync<bSkipDynamics>(*Handle, FrameAndPhase, PropertiesPool))
 	{
 		if (!SkipDesyncTest)
@@ -357,31 +437,30 @@ void FRewindData::AdvanceFrameImp(IResimCacheBase* ResimCache)
 	auto AdvanceHelper = [this, EarliestFrame, FrameAndPhase](auto& DirtyObjects, const auto& DesyncFunc, const auto& AdvanceDirtyFunc)
 	{
 		for (int32 DirtyIdx = DirtyObjects.Num() - 1; DirtyIdx >= 0; --DirtyIdx)
-	{
-			auto& Info = DirtyObjects.GetDenseAt(DirtyIdx);
-
-		ensure(IsResimAndInSync(*Info.GetObjectPtr()) || Info.GetHistory().IsClean(FrameAndPhase));  //Sim hasn't run yet so PostCallbacks (sim results) should be clean
-
-		//if hasn't changed in a while stop tracking
-		if (Info.LastDirtyFrame < EarliestFrame)
-		{
-				RemoveObject(Info.GetObjectPtr());
-		}
-		else
 		{
+			auto& Info = DirtyObjects.GetDenseAt(DirtyIdx);
 
-			auto Handle = Info.GetObjectPtr();
-			Info.bResimAsFollower = Handle->ResimType() == EResimType::ResimAsFollower;
+			ensure(IsResimAndInSync(*Info.GetObjectPtr()) || Info.GetHistory().IsClean(FrameAndPhase));  //Sim hasn't run yet so PostCallbacks (sim results) should be clean
 
-			if (IsResim() && !Info.bResimAsFollower)
+			//if hasn't changed in a while stop tracking
+			if (Info.LastDirtyFrame < EarliestFrame)
 			{
-					DesyncIfNecessary</*bSkipDynamics=*/false>(Info, FrameAndPhase);
+				RemoveObject(Info.GetObjectPtr());
 			}
+			else
+			{
+				auto Handle = Info.GetObjectPtr();
+				Info.bResimAsFollower = Handle->ResimType() == EResimType::ResimAsFollower;
+
+				if (IsResim() && !Info.bResimAsFollower)
+				{
+						DesyncIfNecessary</*bSkipDynamics=*/false>(Info, FrameAndPhase);
+				}
 
 				if (IsResim() && Handle->SyncState() != ESyncState::InSync && !SkipDesyncTest)
-			{
-				Handle->SetEnabledDuringResim(true);	//for now just mark anything out of sync as resim enabled. TODO: use bubble
-					DesyncFunc(Handle);
+				{
+					Handle->SetEnabledDuringResim(true);	//for now just mark anything out of sync as resim enabled. TODO: use bubble
+						DesyncFunc(Handle);
 				}
 
 				AdvanceDirtyFunc(Info, Handle);
@@ -393,7 +472,7 @@ void FRewindData::AdvanceFrameImp(IResimCacheBase* ResimCache)
 	if (IsResim())
 	{
 		DesyncedParticles.Reserve(DirtyParticles.Num());
-			}
+	}
 
 	AdvanceHelper(DirtyParticles,
 		[&DesyncedParticles](FGeometryParticleHandle* DesyncedHandle)
@@ -498,8 +577,13 @@ void FRewindData::PushGTDirtyData(const FDirtyPropertiesManager& SrcManager,cons
 	case EPhysicsProxyType::SingleParticleProxy:
 	{
 		auto ParticleProxy = static_cast<FSingleParticlePhysicsProxy*>(Dirty.Proxy);
+
 		FGeometryParticleHandle* PTParticle = ParticleProxy->GetHandle_LowLevel();
 
+		if (PTParticle == nullptr)
+		{
+			break;
+		}
 		const bool bKeepRecording = CopyHelper(PTParticle, [PTParticle, &DirtyPropHelper](FGeometryParticleStateBase& Latest)
 		{
 			DirtyPropHelper(Latest.ParticlePositionRotation, EChaosPropertyFlags::XR, *PTParticle);
@@ -639,6 +723,53 @@ void FRewindData::MarkDirtyJointFromPT(FPBDJointConstraintHandle& Handle)
 	}
 }
 
+void FRewindData::ClearPhaseAndFuture(FGeometryParticleHandle& Handle, int32 Frame, FFrameAndPhase::EParticleHistoryPhase Phase)
+{
+	FDirtyParticleInfo& Info = FindOrAddDirtyObj(Handle);
+	const FFrameAndPhase FrameAndPhase{ Frame, Phase };
+	Info.ClearPhaseAndFuture(FrameAndPhase);
+}
+
+void FRewindData::ExtendHistoryWithFrame(const int32 Frame)
+{
+	FramesSaved = FMath::Max(CurFrame - Frame+1, FramesSaved);
+}
+
+
+void FRewindData::PushStateAtFrame(FGeometryParticleHandle& Handle, int32 Frame, FFrameAndPhase::EParticleHistoryPhase Phase, 
+	const FVector& Position, const FQuat& Quaternion, const FVector& LinVelocity, const FVector& AngVelocity, const bool bShouldSleep)
+{
+	const bool bRecordingHistory = !IsResimAndInSync(Handle);
+
+	FDirtyParticleInfo& Info = FindOrAddDirtyObj(Handle);
+	const int32 LastFrame = Info.LastDirtyFrame;
+	FGeometryParticleStateBase& Latest = Info.AddFrame(Frame);
+	const FFrameAndPhase FrameAndPhase{ Frame, Phase };
+
+	if (bRecordingHistory || Latest.TargetPositions.IsClean(FrameAndPhase))
+	{
+		FParticlePositionRotation& PositionRotation = Latest.TargetPositions.WriteAccessMonotonic(FrameAndPhase, PropertiesPool);
+		PositionRotation.SetX(Position);
+		PositionRotation.SetR(Quaternion);
+	}
+
+	if (bRecordingHistory || Latest.TargetVelocities.IsClean(FrameAndPhase))
+	{
+		FParticleVelocities& PreVelocities = Latest.TargetVelocities.WriteAccessMonotonic(FrameAndPhase, PropertiesPool);
+		PreVelocities.SetV(LinVelocity);
+		PreVelocities.SetW(AngVelocity);
+	}
+	if (bRecordingHistory || Latest.DynamicsMisc.IsClean(FrameAndPhase))
+	{
+		if (bShouldSleep)
+		{
+			FParticleDynamicMisc& PreDynamicsMisc = Latest.DynamicsMisc.WriteAccessMonotonic(FrameAndPhase, PropertiesPool);
+			PreDynamicsMisc.SetObjectState(EObjectStateType::Sleeping);
+		}
+	}
+	Info.LastDirtyFrame = LastFrame;
+}
+
 void FRewindData::PushPTDirtyData(TPBDRigidParticleHandle<FReal,3>& Handle,const int32 SrcDataIdx)
 {
 	const bool bRecordingHistory = !IsResimAndInSync(Handle);
@@ -679,5 +810,24 @@ FJointState FRewindData::GetPastJointStateAtFrame(const FPBDJointConstraintHandl
 	return GetPastStateAtFrameImp<FJointState>(DirtyJoints, Handle, Frame, Phase);
 }
 
+void FRewindData::SetTargetStateAtFrame(FGeometryParticleHandle& Handle, const int32 Frame, FFrameAndPhase::EParticleHistoryPhase Phase,
+	const FVector& Position, const FQuat& Quaternion, const FVector& LinVelocity, const FVector& AngVelocity, const bool bShouldSleep)
+{
+	int32 PositionValidCount = INDEX_NONE;
+	int32 VelocityValidCount = INDEX_NONE;
+	int32 PositionNextIterator = INDEX_NONE;
+	int32 VelocityNextIterator = INDEX_NONE;
+
+	ExtractHistoryState(Handle,
+		PositionValidCount, VelocityValidCount, PositionNextIterator, VelocityNextIterator);
+
+	ClearPhaseAndFuture(Handle, Frame, Phase);
+
+	PushStateAtFrame(Handle, Frame, Phase,
+		Position, Quaternion, LinVelocity, AngVelocity, bShouldSleep);
+
+	RestoreHistoryState(Handle,
+		PositionValidCount, VelocityValidCount, PositionNextIterator, VelocityNextIterator);
+}
 
 }
diff --git a/Engine/Source/Runtime/Experimental/Chaos/Public/Chaos/ChaosMarshallingManager.h b/Engine/Source/Runtime/Experimental/Chaos/Public/Chaos/ChaosMarshallingManager.h
--- a/Engine/Source/Runtime/Experimental/Chaos/Public/Chaos/ChaosMarshallingManager.h
+++ b/Engine/Source/Runtime/Experimental/Chaos/Public/Chaos/ChaosMarshallingManager.h
@@ -340,6 +340,9 @@ public:
 
 	/** Returns the history buffer of the latest NumFrames. The history that comes before these frames is discarded*/
 	TArray<FPushPhysicsData*> StealHistory_Internal(int32 NumFrames);
+
+	/** Return the size of the history queue */
+	int32 GetNumHistory_Internal() const {return HistoryQueue_Internal.Num();}
 		
 private:
 	FReal ExternalTime_External;	//the global time external thread is currently at
diff --git a/Engine/Source/Runtime/Experimental/Chaos/Public/Chaos/Framework/PhysicsSolverBase.h b/Engine/Source/Runtime/Experimental/Chaos/Public/Chaos/Framework/PhysicsSolverBase.h
--- a/Engine/Source/Runtime/Experimental/Chaos/Public/Chaos/Framework/PhysicsSolverBase.h
+++ b/Engine/Source/Runtime/Experimental/Chaos/Public/Chaos/Framework/PhysicsSolverBase.h
@@ -429,6 +429,7 @@ namespace Chaos
 		void DisableAsyncMode();
 		
 		virtual void ConditionalApplyRewind_Internal(){}
+		virtual bool IsResimming() const {return false;}
 
 		FChaosMarshallingManager& GetMarshallingManager() { return MarshallingManager; }
 		FChaosResultsManager& GetResultsManager() { return *PullResultsManager; }
@@ -587,6 +588,41 @@ namespace Chaos
 		/**/
 		FReal GetLastDt() const { return MLastDt; }
 
+		/** Return the interpolation lerp in case the resim is off */
+		static float NetworkPhysicsInterpolationLerp()
+		{
+			static IConsoleVariable* InterpLerpCVar = IConsoleManager::Get().FindConsoleVariable(TEXT("np2.NetworkPhysicsPredictionInterpLerp"));
+			return InterpLerpCVar ? InterpLerpCVar->GetFloat() : 0.1f;
+		}
+
+		/** Return the error position threshold to trigger a resim*/
+		static float NetworkPhysicsErrorThreshold()
+		{
+			static IConsoleVariable* ErrorThresholdCVar = IConsoleManager::Get().FindConsoleVariable(TEXT("np2.NetworkPhysicsPredictionErrorThreshold"));
+			return ErrorThresholdCVar ? ErrorThresholdCVar->GetFloat() : 10.0f;
+		}
+
+		/** Check if network physics is enables or not */
+		static bool IsNetworkPhysicsPredictionEnabled()
+		{
+			static IConsoleVariable* EnableNetworkPhysicsCVar = IConsoleManager::Get().FindConsoleVariable(TEXT("np2.EnableNetworkPhysicsPrediction"));
+			return EnableNetworkPhysicsCVar && (EnableNetworkPhysicsCVar->GetInt() == 1);
+		}
+
+		/** Check if we can enable debugging informations for network physics */
+		static bool CanDebugNetworkPhysicsPrediction()
+		{
+			static IConsoleVariable* DebugNetworkPhysicsCVar = IConsoleManager::Get().FindConsoleVariable(TEXT("np2.DebugNetworkPhysicsPrediction"));
+			return DebugNetworkPhysicsCVar && (DebugNetworkPhysicsCVar->GetInt() == 1);
+		}
+
+		/** Check if resim is enabled for network physics */
+		static bool CanResimNetworkPhysicsPrediction()
+		{
+			static IConsoleVariable* ResimNetworkPhysicsCVar = IConsoleManager::Get().FindConsoleVariable(TEXT("np2.ResimNetworkPhysicsPrediction"));
+			return ResimNetworkPhysicsCVar && (ResimNetworkPhysicsCVar->GetInt() == 1);
+		}
+
 	protected:
 		/** Mode that the results buffers should be set to (single, double, triple) */
 		EMultiBufferMode BufferMode;
@@ -714,3 +750,5 @@ namespace Chaos
 #endif
 	};
 }
+
+
diff --git a/Engine/Source/Runtime/Experimental/Chaos/Public/Chaos/Island/IslandGraph.h b/Engine/Source/Runtime/Experimental/Chaos/Public/Chaos/Island/IslandGraph.h
--- a/Engine/Source/Runtime/Experimental/Chaos/Public/Chaos/Island/IslandGraph.h
+++ b/Engine/Source/Runtime/Experimental/Chaos/Public/Chaos/Island/IslandGraph.h
@@ -42,6 +42,9 @@ namespace Chaos
 
 		/** Set of used colors */
 		TSet<int32> ColorIndices;
+
+		/** Resim frame */
+		int32 ResimFrame = INDEX_NONE;
 	};
 
 	/** Graph edge structure */
diff --git a/Engine/Source/Runtime/Experimental/Chaos/Public/Chaos/Island/IslandGroupManager.h b/Engine/Source/Runtime/Experimental/Chaos/Public/Chaos/Island/IslandGroupManager.h
--- a/Engine/Source/Runtime/Experimental/Chaos/Public/Chaos/Island/IslandGroupManager.h
+++ b/Engine/Source/Runtime/Experimental/Chaos/Public/Chaos/Island/IslandGroupManager.h
@@ -116,7 +116,7 @@ namespace Chaos
 			 * Pull all the active islands from the IslandManager and assign to groups.
 			 * @return 
 			*/
-			int32 BuildGroups();
+			int32 BuildGroups(const bool bIsResimming);
 
 			/**
 			 * Set the default number of iterations to perform in the Solve step (can be increased on a per-island basis by any dynamic body).
diff --git a/Engine/Source/Runtime/Experimental/Chaos/Public/Chaos/Island/IslandManager.h b/Engine/Source/Runtime/Experimental/Chaos/Public/Chaos/Island/IslandManager.h
--- a/Engine/Source/Runtime/Experimental/Chaos/Public/Chaos/Island/IslandManager.h
+++ b/Engine/Source/Runtime/Experimental/Chaos/Public/Chaos/Island/IslandManager.h
@@ -193,7 +193,7 @@ public:
 	/**
 	  * Pack all the constraints into their assigned islands in sorted order
 	  */
-	void UpdateIslands(FPBDRigidsSOAs& Particles);
+	void UpdateIslands(FPBDRigidsSOAs& Particles, const bool IsResimming = false);
 
 	/**
 	  * @brief Put all particles and constraints in an island to sleep
@@ -233,6 +233,15 @@ public:
 	 */
 	bool IslandNeedsResim(const int32 IslandIndex) const;
 
+	/** Get the particle resim frame */
+	int32 GetParticleResimFrame(const FGeometryParticleHandle* ParticleHandle) const;
+
+	/** Set the particle resim frame */
+	void SetParticleResimFrame(FGeometryParticleHandle* ParticleHandle, const int32 ResimFrame);
+
+	/** Reset all the particle resim frames to INDEX_NONE */
+	void ResetParticleResimFrame(const int32 ResetFrame = INDEX_NONE);
+
 	/**
 	  * The number of islands in the graph.
 	  * @return the number of islands
@@ -340,7 +349,7 @@ protected:
 	/**
 	* Sync the Islands with the IslandGraph
 	*/
-	void SyncIslands(FPBDRigidsSOAs& Particles);
+	void SyncIslands(FPBDRigidsSOAs& Particles, const bool IsResimming);
 
 
 	/**
diff --git a/Engine/Source/Runtime/Experimental/Chaos/Public/Chaos/Island/SolverIsland.h b/Engine/Source/Runtime/Experimental/Chaos/Public/Chaos/Island/SolverIsland.h
--- a/Engine/Source/Runtime/Experimental/Chaos/Public/Chaos/Island/SolverIsland.h
+++ b/Engine/Source/Runtime/Experimental/Chaos/Public/Chaos/Island/SolverIsland.h
@@ -170,6 +170,16 @@ namespace Chaos
 			*/
 			FORCEINLINE int32 GetNumConstraints() const { return NumConstraints; }
 
+			/**
+			* Get the resim frame value
+			*/
+			FORCEINLINE int32 GetResimFrame() const { return ResimFrame; }
+
+			/**
+			* Set the resim frame value
+			*/
+			FORCEINLINE void SetResimFrame(const int32 InResimFrame) { ResimFrame = InResimFrame; }
+
 			/**
 			* Members accessors
 			*/
@@ -177,7 +187,6 @@ namespace Chaos
 			FORCEINLINE void SetIsSleeping(const bool bInIsSleepingIn) { bIsSleeping = bInIsSleepingIn; }
 			FORCEINLINE bool IsSleepingChanged() const { return bIsSleepingChanged; }
 			FORCEINLINE void SetIsSleepingChanged(const bool bInIsSleepingChanged) { bIsSleepingChanged = bInIsSleepingChanged; }
-
 			FORCEINLINE bool IsPersistent() const { return bIsPersistent; }
 			FORCEINLINE void SetIsPersistent(const bool bIsPersistentIn) { bIsPersistent = bIsPersistentIn; }
 			FORCEINLINE bool NeedsResim() const { return bNeedsResim; }
@@ -239,6 +248,9 @@ namespace Chaos
 			/** Sleep counter to trigger island sleeping */
 			int32 SleepCounter = 0;
 
+			/** Resim frame for the island */
+			int32 ResimFrame = INDEX_NONE;
+
 			/** List of all the island particles handles */
 			TArray<FPBDIslandParticle> IslandParticles;
 
diff --git a/Engine/Source/Runtime/Experimental/Chaos/Public/Chaos/PBDRigidsEvolution.h b/Engine/Source/Runtime/Experimental/Chaos/Public/Chaos/PBDRigidsEvolution.h
--- a/Engine/Source/Runtime/Experimental/Chaos/Public/Chaos/PBDRigidsEvolution.h
+++ b/Engine/Source/Runtime/Experimental/Chaos/Public/Chaos/PBDRigidsEvolution.h
@@ -342,6 +342,10 @@ public:
 	CHAOS_API FPBDRigidsSOAs& GetParticles() { return Particles; }
 	CHAOS_API const FPBDRigidsSOAs& GetParticles() const { return Particles; }
 
+	/** Reset the collisions warm starting when resimulate. Ideally we should store
+		  that in the RewindData history but probably too expensive for now */
+	virtual void ResetCollisions() {};
+
 	/**
 	* Register a constraint container with the evolution. Constraints added to the container will be solved during the tick.
 	* @note we do not currently support removing containers. In a few places we assume the ContainerId is persistent and equal to the array index.
@@ -884,6 +888,9 @@ public:
 	void SetResim(bool bInResim) { bIsResim = bInResim; }
 	const bool IsResimming() const { return bIsResim; }
 
+	void SetReset(bool bInReset) { bIsReset = bInReset; }
+	const bool IsResetting() const { return bIsReset; }
+
 	void Serialize(FChaosArchive& Ar);
 
 	FUniqueIdx GenerateUniqueIdx()
@@ -953,6 +960,7 @@ public:
 		InternalAcceleration->RemoveElementFrom(SpatialData.AccelerationHandle, SpatialData.SpatialIdx);
 	}
 
+	
 protected:
 
 	void UpdateConstraintPositionBasedState(FReal Dt)
@@ -987,9 +995,9 @@ protected:
 	void CreateIslands()
 	{
 		// Package the constraints and particles into islands
-		IslandManager.UpdateIslands(Particles);
+		IslandManager.UpdateIslands(Particles, IsResimming());
 	}
-	
+
 	void FlushInternalAccelerationQueue();
 	void FlushAsyncAccelerationQueue();
 	void WaitOnAccelerationStructure();
@@ -1106,6 +1114,7 @@ protected:
 
 	TArray<FUniqueIdx> PendingReleaseIndices;	//for now just assume a one frame delay, but may need something more general
 	bool bIsResim = false; 
+	bool bIsReset = false;
 };
 
 
diff --git a/Engine/Source/Runtime/Experimental/Chaos/Public/Chaos/PBDRigidsEvolutionGBF.h b/Engine/Source/Runtime/Experimental/Chaos/Public/Chaos/PBDRigidsEvolutionGBF.h
--- a/Engine/Source/Runtime/Experimental/Chaos/Public/Chaos/PBDRigidsEvolutionGBF.h
+++ b/Engine/Source/Runtime/Experimental/Chaos/Public/Chaos/PBDRigidsEvolutionGBF.h
@@ -171,6 +171,10 @@ namespace Chaos
 
 		virtual void DestroyTransientConstraints(FGeometryParticleHandle* Particle) override;
 
+		/** Reset the collisions warm starting when resimulate. Ideally we should store
+		  that in the RewindData history but probably too expensive for now */
+		virtual void ResetCollisions() override;
+
 		CHAOS_API inline void EndFrame(FReal Dt)
 		{
 			Particles.GetNonDisabledDynamicView().ParallelFor([&](auto& Particle, int32 Index) {
@@ -212,13 +216,11 @@ namespace Chaos
 					Particle.V() += Particle.Acceleration() * Dt;
 					Particle.W() += Particle.AngularAcceleration() * Dt;
 
-
 					//AddImpulsesRule.Apply(Particle, Dt);
 					Particle.V() += Particle.LinearImpulseVelocity();
 					Particle.W() += Particle.AngularImpulseVelocity();
 					Particle.LinearImpulseVelocity() = FVec3(0);
 					Particle.AngularImpulseVelocity() = FVec3(0);
-					
 
 					//EtherDragRule.Apply(Particle, Dt);
 					{
diff --git a/Engine/Source/Runtime/Experimental/Chaos/Public/PBDRigidsSolver.h b/Engine/Source/Runtime/Experimental/Chaos/Public/PBDRigidsSolver.h
--- a/Engine/Source/Runtime/Experimental/Chaos/Public/PBDRigidsSolver.h
+++ b/Engine/Source/Runtime/Experimental/Chaos/Public/PBDRigidsSolver.h
@@ -395,6 +395,9 @@ namespace Chaos
 		void DestroyPendingProxies_Internal();
 
 		virtual void ConditionalApplyRewind_Internal() override;
+
+		/** Check if we are resimming or not */
+		virtual bool IsResimming() const {return GetEvolution()->IsResimming();}
 	};
 
 	template<>
diff --git a/Engine/Source/Runtime/Experimental/Chaos/Public/RewindData.h b/Engine/Source/Runtime/Experimental/Chaos/Public/RewindData.h
--- a/Engine/Source/Runtime/Experimental/Chaos/Public/RewindData.h
+++ b/Engine/Source/Runtime/Experimental/Chaos/Public/RewindData.h
@@ -15,6 +15,142 @@
 namespace Chaos
 {
 
+/** Base states history interface to store additional information required for rewind */
+struct FBaseStatesHistory
+{
+	FORCEINLINE virtual ~FBaseStatesHistory() {}
+
+	/** Rewind the states at some point in time */
+	FORCEINLINE virtual bool RewindStates(const int32 RewindFrame, const bool bResetSolver) { return true; }
+};
+
+/** Base Inputs history interface to continuously apply Inputs inputs while resimulating */
+struct FBaseInputsHistory
+{
+	FORCEINLINE virtual ~FBaseInputsHistory() {}
+
+	/** Apply Inputs given a time */
+	FORCEINLINE virtual bool ApplyInputs(const int32 ApplyFrame, const bool bResetSolver) { return true; }
+};
+
+/** Base history class that is holding a fixed size buffer of datas */
+template<typename DatasType>
+struct FDatasRewindHistory 
+{
+	FORCEINLINE FDatasRewindHistory(const int32 FrameCount, const bool bIsHistoryLocal) : 
+		bIsLocalHistory(bIsHistoryLocal), DatasArray(), CurrentFrame(0), CurrentIndex(0), NumFrames(FrameCount)
+	{
+		DatasArray.SetNum(NumFrames);
+	}
+	FORCEINLINE virtual ~FDatasRewindHistory() {}
+
+protected:
+
+	/** Load the datas from the buffer at a specific frame */
+	FORCEINLINE bool LoadDatas(const int32 LoadFrame)
+	{
+		const int32 LocalFrame = LoadFrame % NumFrames;
+		if (LoadFrame == DatasArray[LocalFrame].DatasFrame)
+		{
+			CurrentFrame = LoadFrame;
+			CurrentIndex = LocalFrame;
+			return true;
+		}
+		return false;
+	}
+
+	/** Read directly the datas from the buffer at a specific frame without chacking the datasframe */
+	FORCEINLINE void ReadDatas(const int32 LoadFrame)
+	{
+		const int32 LocalFrame = LoadFrame % NumFrames;
+		CurrentFrame = LoadFrame;
+		CurrentIndex = LocalFrame;
+	}
+
+	/** Record the datas from the buffer at a specific frame */
+	FORCEINLINE virtual bool RecordDatas(const int32 RecordFrame)
+	{
+		const int32 LocalFrame = RecordFrame % NumFrames;
+		DatasArray[LocalFrame].DatasFrame = RecordFrame;
+		CurrentFrame = RecordFrame;
+		CurrentIndex = LocalFrame;
+		return true;
+	}
+
+	/** Current datas that is being loaded/recorded*/
+	DatasType& GetCurrentDatas() { return DatasArray[CurrentIndex]; }
+	const DatasType& GetCurrentDatas() const { return DatasArray[CurrentIndex]; }
+
+	/** Check if the history is on the local/remote client*/
+	bool bIsLocalHistory;
+
+	/**  Datas buffer holding the history */
+	TArray<DatasType> DatasArray;
+
+	/** Current frame that is being loaded/recorded */
+	int32 CurrentFrame;
+
+	/** Current index that is being loaded/recorded */
+	int32 CurrentIndex;
+
+	/** Number of frames used in that history */
+	int32 NumFrames = 0;
+};
+
+/** Base history class for the states rewind datas */
+template<typename DatasType>
+struct  FStatesRewindHistory : public FDatasRewindHistory<DatasType>, public FBaseStatesHistory
+{
+	using DatasSuper = FDatasRewindHistory<DatasType>;
+
+	FORCEINLINE FStatesRewindHistory(const int32 FrameCount, const bool bIsHistoryLocal) :
+		FDatasRewindHistory<DatasType>(FrameCount, bIsHistoryLocal), FBaseStatesHistory()
+	{}
+
+	FORCEINLINE virtual ~FStatesRewindHistory() override {}
+
+	/** Rewind the states */
+	FORCEINLINE virtual bool RewindStates(const int32 RewindFrame, const bool bResetSolver) override
+	{
+		if (!DatasSuper::LoadDatas(RewindFrame))
+		{
+			if(bResetSolver)
+			{
+				UE_LOG(LogChaos, Warning, TEXT("Unable to read target state at frame %d while rewinding the simulation"), RewindFrame);
+			}
+			return false;
+		}
+		return true;
+	}
+};
+
+/** Base history class for the Inputs rewind datas */
+template<typename DatasType>
+struct  FInputsRewindHistory : public FDatasRewindHistory<DatasType>, public FBaseInputsHistory
+{
+	using DatasSuper = FDatasRewindHistory<DatasType>;
+
+	FORCEINLINE FInputsRewindHistory(const int32 FrameCount, const bool bIsHistoryLocal) :
+		FDatasRewindHistory<DatasType>(FrameCount, bIsHistoryLocal), FBaseInputsHistory()
+	{}
+
+	FORCEINLINE virtual ~FInputsRewindHistory() override {}
+
+	/** Apply the Inputs */
+	FORCEINLINE virtual bool ApplyInputs(const int32 ApplyFrame, const bool bResetSolver) override
+	{
+		if (!DatasSuper::LoadDatas(ApplyFrame))
+		{
+			if (DatasSuper::bIsLocalHistory || (!DatasSuper::bIsLocalHistory && bResetSolver))
+			{
+				UE_LOG(LogChaos, Warning, TEXT("Unable to apply Inputs on the local client at frame %d"), ApplyFrame);
+			}
+			return false;
+		}
+		return true;
+	}
+};
+
 struct FFrameAndPhase
 {
 	enum EParticleHistoryPhase : uint8
@@ -160,6 +296,11 @@ public:
 		NumValid = 0;
 	}
 
+	bool IsEmpty() const
+	{
+		return NumValid == 0;
+	}
+
 	void ClearEntryAndFuture(const FFrameAndPhase FrameAndPhase)
 	{
 		//Move next backwards until FrameAndPhase and anything more future than it is gone
@@ -168,7 +309,7 @@ public:
 			const int32 PotentialNext = Next - 1 >= 0 ? Next - 1 : Buffer.Num() - 1;
 
 			if(Buffer[PotentialNext].FrameAndPhase < FrameAndPhase)
-	{
+			{
 				break;
 			}
 
@@ -177,6 +318,18 @@ public:
 		}
 	}
 
+	void ExtractBufferState(int32& ValidCount, int32& NextIterator) const
+	{
+		ValidCount = NumValid;
+		NextIterator = Next;
+	}
+
+	void RestoreBufferState(const int32& ValidCount, const int32& NextIterator)
+	{
+		NumValid = ValidCount;
+		Next = NextIterator;
+	}
+
 	bool IsClean(const FFrameAndPhase FrameAndPhase) const
 	{
 		return FindIdx(FrameAndPhase) == INDEX_NONE;
@@ -538,6 +691,8 @@ struct FGeometryParticleStateBase
 	, DynamicsMisc(ComputeCircularSize(NumFrames))
 	, MassProps(ComputeCircularSize(NumFrames))
 	, KinematicTarget(ComputeCircularSize(NumFrames))
+	, TargetPositions(ComputeCircularSize(NumFrames))
+	, TargetVelocities(ComputeCircularSize(NumFrames))
 	{
 
 	}
@@ -555,6 +710,8 @@ struct FGeometryParticleStateBase
 		DynamicsMisc.Release(Manager);
 		MassProps.Release(Manager);
 		KinematicTarget.Release(Manager);
+		TargetPositions.Release(Manager);
+		TargetVelocities.Release(Manager);
 	}
 
 	void Reset()
@@ -566,6 +723,8 @@ struct FGeometryParticleStateBase
 		DynamicsMisc.Reset();
 		MassProps.Reset();
 		KinematicTarget.Reset();
+		TargetVelocities.Reset();
+		TargetPositions.Reset();
 	}
 
 	void ClearEntryAndFuture(const FFrameAndPhase FrameAndPhase)
@@ -577,6 +736,20 @@ struct FGeometryParticleStateBase
 		DynamicsMisc.ClearEntryAndFuture(FrameAndPhase);
 		MassProps.ClearEntryAndFuture(FrameAndPhase);
 		KinematicTarget.ClearEntryAndFuture(FrameAndPhase);
+		TargetPositions.ClearEntryAndFuture(FrameAndPhase);
+		TargetVelocities.ClearEntryAndFuture(FrameAndPhase);
+	}
+
+	void ExtractHistoryState(int32& PositionValidCount, int32& VelocityValidCount, int32& PositionNextIterator, int32& VelocityNextIterator) const
+	{
+		ParticlePositionRotation.ExtractBufferState(PositionValidCount, PositionNextIterator);
+		Velocities.ExtractBufferState(VelocityValidCount, VelocityNextIterator);
+	}
+
+	void RestoreHistoryState(const int32& PositionValidCount, const int32& VelocityValidCount, const int32& PositionNextIterator, const int32& VelocityNextIterator)
+	{
+		ParticlePositionRotation.RestoreBufferState(PositionValidCount, PositionNextIterator);
+		Velocities.RestoreBufferState(VelocityValidCount, VelocityNextIterator);
 	}
 
 	bool IsClean(const FFrameAndPhase FrameAndPhase) const
@@ -596,6 +769,9 @@ struct FGeometryParticleStateBase
 
 	template <bool bSkipDynamics = false>
 	bool IsInSync(const FGeometryParticleHandle& Handle, const FFrameAndPhase FrameAndPhase, const FDirtyPropertiesPool& Pool) const;
+
+	/** Check if the handle resim frame is valid (before the current one) */
+	bool IsResimFrameValid(const FGeometryParticleHandle& Handle, const FFrameAndPhase FrameAndPhase) const;
 	
 	template <typename TParticle>
 	static TShapesArrayState<TParticle> ShapesArray(const FGeometryParticleStateBase* State, const TParticle& Particle)
@@ -614,6 +790,9 @@ struct FGeometryParticleStateBase
 	TParticlePropertyBuffer<FParticleMassProps,EChaosProperty::MassProps> MassProps;
 	TParticlePropertyBuffer<FKinematicTarget, EChaosProperty::KinematicTarget> KinematicTarget;
 
+	TParticlePropertyBuffer<FParticlePositionRotation, EChaosProperty::XR> TargetPositions;
+	TParticlePropertyBuffer<FParticleVelocities, EChaosProperty::Velocities> TargetVelocities;
+
 	FShapesArrayStateBase ShapesArrayState;
 };
 
@@ -826,6 +1005,7 @@ private:
 template <typename T> 
 const T* ConstifyHelper(T* Ptr) { return Ptr; }
 
+
 template <typename T>
 T NoRefHelper(const T& Ref) { return Ref; }
 
@@ -937,6 +1117,14 @@ public:
 	{
 	}
 
+	void Init(FPBDRigidsSolver* InSolver, int32 NumFrames, bool InResimOptimization, int32 InCurrentFrame)
+	{
+		Solver = InSolver;
+		CurFrame = InCurrentFrame;
+		bResimOptimization = InResimOptimization;
+		Managers = TCircularBuffer<FFrameManagerInfo>(NumFrames + 1);
+	}
+
 	int32 Capacity() const { return Managers.Capacity(); }
 	int32 CurrentFrame() const { return CurFrame; }
 	int32 GetFramesSaved() const { return FramesSaved; }
@@ -959,6 +1147,33 @@ public:
 
 	int32 CHAOS_API GetEarliestFrame_Internal() const { return CurFrame - FramesSaved; }
 
+	/* Extend the current history size to be sure to include the given frame */
+	void CHAOS_API ExtendHistoryWithFrame(const int32 Frame);
+
+	/* Clear all the simulation history after Frame */
+	void CHAOS_API ClearPhaseAndFuture(FGeometryParticleHandle& Handle, int32 Frame, FFrameAndPhase::EParticleHistoryPhase Phase);
+
+	/* Push a physics state in the rewind datas at some frame */
+	void CHAOS_API PushStateAtFrame(FGeometryParticleHandle& Handle, int32 Frame, FFrameAndPhase::EParticleHistoryPhase Phase, const FVector& Position, const FQuat& Quaternion,
+					const FVector& LinVelocity, const FVector& AngVelocity, const bool bShouldSleep);
+
+	void CHAOS_API SetTargetStateAtFrame(FGeometryParticleHandle& Handle, const int32 Frame, FFrameAndPhase::EParticleHistoryPhase Phase,
+		const FVector& Position, const FQuat& Quaternion, const FVector& LinVelocity, const FVector& AngVelocity, const bool bShouldSleep);
+
+	/** Extract some history information before cleaning/pushing state*/
+	void CHAOS_API ExtractHistoryState(FGeometryParticleHandle& Handle, int32& PositionValidCount, int32& VelocityValidCount, int32& PositionNextIterator, int32& VelocityNextIterator)
+	{
+		FDirtyParticleInfo& Info = FindOrAddDirtyObj(Handle);
+		Info.GetHistory().ExtractHistoryState(PositionValidCount, VelocityValidCount, PositionNextIterator, VelocityNextIterator);
+	}
+
+	/** Restore some history information after cleaning/pushing state*/
+	void CHAOS_API RestoreHistoryState(FGeometryParticleHandle& Handle, const int32& PositionValidCount, const int32& VelocityValidCount, const int32& PositionNextIterator, const int32& VelocityNextIterator)
+	{
+		FDirtyParticleInfo& Info = FindOrAddDirtyObj(Handle);
+		Info.GetHistory().RestoreHistoryState(PositionValidCount, VelocityValidCount, PositionNextIterator, VelocityNextIterator);
+	}
+
 	/* Query the state of particles from the past. Can only be used when not already resimming*/
 	FGeometryParticleState CHAOS_API GetPastStateAtFrame(const FGeometryParticleHandle& Handle, int32 Frame, FFrameAndPhase::EParticleHistoryPhase Phase = FFrameAndPhase::EParticleHistoryPhase::PostPushData) const;
 
@@ -1033,6 +1248,39 @@ public:
 
 	void CHAOS_API SpawnProxyIfNeeded(FSingleParticlePhysicsProxy& Proxy);
 
+	/** Add Inputs history to the rewind datas for future use while resimulating */
+	void CHAOS_API AddInputsHistory(const TSharedPtr<FBaseInputsHistory>& InputsHistory)
+	{
+		InputsHistories.Add(InputsHistory.ToWeakPtr());
+	}
+
+	/** Remove Inputs history from the rewind datasg */
+	void CHAOS_API RemoveInputsHistory(const TSharedPtr<FBaseInputsHistory>& InputsHistory)
+	{
+		InputsHistories.Remove(InputsHistory.ToWeakPtr());
+	}
+
+	/** Add states history to the rewind datas for future use while rewinding */
+	void CHAOS_API AddStatesHistory(const TSharedPtr<FBaseStatesHistory>& StatesHistory)
+	{
+		StatesHistories.Add(StatesHistory.ToWeakPtr());
+	}
+
+	/** Remove states history from the rewind datas */
+	void CHAOS_API RemoveStatesHistory(const TSharedPtr<FBaseStatesHistory>& StatesHistory)
+	{
+		StatesHistories.Remove(StatesHistory.ToWeakPtr());
+	}
+
+	/** Apply the inputs from all the history Inputs datas */
+	void ApplyInputs(const int32 ApplyFrame, const bool bResetSolver);
+
+	/** Rewind the states from all the history states datas */
+	void RewindStates(const int32 RewindFrame, const bool bResetSolver);
+
+	/** Return the rewind data solver */
+	const FPBDRigidsSolver* GetSolver() const { return Solver; }
+
 private:
 
 	friend class FPBDRigidsSolver;
@@ -1111,6 +1359,11 @@ private:
 		{
 			return History;
 		}
+
+		THistoryType& GetHistory()
+		{
+			return History;
+		}
 	};
 
 	using FDirtyParticleInfo = TDirtyObjectInfo<FGeometryParticleStateBase, FGeometryParticleHandle>;
@@ -1151,7 +1404,10 @@ private:
 		return TObjState(State, Handle, PropertiesPool, { Frame, Phase });
 	}
 
-	bool RewindToFrame(int32 Frame);
+	bool RewindToFrame(int32 RewindFrame);
+
+	/** Apply targets positions and velocities while resimulating */
+	void ApplyTargets(const int32 Frame, const bool bResetSimulation);
 
 	template <typename TDirtyInfo>
 	static void DesyncObject(TDirtyInfo& Info, const FFrameAndPhase FrameAndPhase)
@@ -1166,6 +1422,9 @@ private:
 	TDirtyObjects<FDirtyParticleInfo> DirtyParticles;
 	TDirtyObjects<FDirtyJointInfo> DirtyJoints;
 
+	TArray<TWeakPtr<FBaseInputsHistory>> InputsHistories;
+	TArray<TWeakPtr<FBaseStatesHistory>> StatesHistories;
+
 	FPBDRigidsSolver* Solver;
 	int32 CurFrame;
 	int32 LatestFrame;
diff --git a/Engine/Source/Runtime/Experimental/ChaosVehicles/ChaosVehiclesCore/Public/EngineSystem.h b/Engine/Source/Runtime/Experimental/ChaosVehicles/ChaosVehiclesCore/Public/EngineSystem.h
--- a/Engine/Source/Runtime/Experimental/ChaosVehicles/ChaosVehiclesCore/Public/EngineSystem.h
+++ b/Engine/Source/Runtime/Experimental/ChaosVehicles/ChaosVehiclesCore/Public/EngineSystem.h
@@ -131,6 +131,17 @@ namespace Chaos
 			return 0.f;
 		}
 
+		/** get the Engine angular velocity */
+		float GetEngineOmega() const
+		{
+			return Omega;
+		}
+
+		/** Set the Engine angular velocity */
+		void SetEngineOmega( const float EngineOmega)
+		{
+			Omega = EngineOmega;
+		}
 
 		/** 
 		 * Simulate - NOP at the moment
diff --git a/Engine/Source/Runtime/Experimental/ChaosVehicles/ChaosVehiclesCore/Public/SuspensionSystem.h b/Engine/Source/Runtime/Experimental/ChaosVehicles/ChaosVehiclesCore/Public/SuspensionSystem.h
--- a/Engine/Source/Runtime/Experimental/ChaosVehicles/ChaosVehiclesCore/Public/SuspensionSystem.h
+++ b/Engine/Source/Runtime/Experimental/ChaosVehicles/ChaosVehiclesCore/Public/SuspensionSystem.h
@@ -205,6 +205,18 @@ namespace Chaos
 
 		void Simulate(float DeltaTime);
 
+		int32 GetAveragingCount() const {return AveragingCount;}
+		int32 GetAveragingNum() const { return AveragingNum; }
+		float GetLastSpringLength() const { return LastSpringLength; }
+		float GetLastDisplacement() const { return LastDisplacement; }
+		float GetAveragingLength(const int32 LengthIndex) const { return AveragingLength[LengthIndex]; }
+
+		void SetAveragingCount(const int32 InAveragingCount) {AveragingCount = InAveragingCount;}
+		void SetAveragingNum(const int32 InAveragingNum) { AveragingNum = InAveragingNum; }
+		void SetLastSpringLength(const float InLastSpringLength) { LastSpringLength = InLastSpringLength; }
+		void SetLastDisplacement(const float InLastDisplacement) { LastDisplacement = InLastDisplacement; }
+		void SetAveragingLength(const int32 LengthIndex, const float InAveragingLength) { AveragingLength[LengthIndex] = InAveragingLength; }
+
 	protected:
 
 		float DisplacementInput;	
diff --git a/Engine/Source/Runtime/Experimental/ChaosVehicles/ChaosVehiclesCore/Public/TransmissionSystem.h b/Engine/Source/Runtime/Experimental/ChaosVehicles/ChaosVehiclesCore/Public/TransmissionSystem.h
--- a/Engine/Source/Runtime/Experimental/ChaosVehicles/ChaosVehiclesCore/Public/TransmissionSystem.h
+++ b/Engine/Source/Runtime/Experimental/ChaosVehicles/ChaosVehiclesCore/Public/TransmissionSystem.h
@@ -137,6 +137,30 @@ namespace Chaos
 			return TargetGear;
 		}
 
+		/** Get the current gear change time */
+		float GetCurrentGearChangeTime() const
+		{
+			return CurrentGearChangeTime;
+		}
+
+		/** Set the current gear index, (reverse gears < 0, neutral 0, forward gears > 0) */
+		void SetCurrentGear(const int32 InCurrentGear) 
+		{
+			CurrentGear = InCurrentGear;
+		}
+
+		/** Set the target gear index, (reverse gears < 0, neutral 0, forward gears > 0) */
+		void SetTargetGear(const int32 InTargetGear) 
+		{
+			TargetGear = InTargetGear;
+		}
+
+		/** Set the current gear change time */
+		void SetCurrentGearChangeTime(const float InCurrentGearChangeTime) 
+		{
+			CurrentGearChangeTime = InCurrentGearChangeTime;
+		}
+
 		/** Are we currently in the middle of a gear change */
 		bool IsCurrentlyChangingGear() const
 		{
@@ -192,14 +216,14 @@ namespace Chaos
 		 * - implements gear change time, where gear goes through neutral
 		 */
 		void Simulate(float DeltaTime);
-
-
-	private:
+	
 		void CorrectGearInputRange(int32& GearIndexInOut)
 		{
 			GearIndexInOut = FMath::Clamp(GearIndexInOut, -Setup().ReverseRatios.Num(), Setup().ForwardRatios.Num());
 		}
 
+	private:
+
 		int32 CurrentGear; // <0 reverse gear(s), 0 neutral, >0 forward gears
 		int32 TargetGear;  // <0 reverse gear(s), 0 neutral, >0 forward gears
 		float CurrentGearChangeTime; // Time to change gear, no power transmitted to the wheels during change
diff --git a/Engine/Source/Runtime/Experimental/GeometryCollectionEngine/Private/GeometryCollection/GeometryCollectionComponent.cpp b/Engine/Source/Runtime/Experimental/GeometryCollectionEngine/Private/GeometryCollection/GeometryCollectionComponent.cpp
--- a/Engine/Source/Runtime/Experimental/GeometryCollectionEngine/Private/GeometryCollection/GeometryCollectionComponent.cpp
+++ b/Engine/Source/Runtime/Experimental/GeometryCollectionEngine/Private/GeometryCollection/GeometryCollectionComponent.cpp
@@ -2852,7 +2852,7 @@ void UGeometryCollectionComponent::RequestUpdateRepData()
 			[this]()
 			{
 				UpdateRepData();
-			}
+			}, false
 		);
 	}
 }
diff --git a/Engine/Source/Runtime/PhysicsCore/Private/ChaosEngineInterface.cpp b/Engine/Source/Runtime/PhysicsCore/Private/ChaosEngineInterface.cpp
--- a/Engine/Source/Runtime/PhysicsCore/Private/ChaosEngineInterface.cpp
+++ b/Engine/Source/Runtime/PhysicsCore/Private/ChaosEngineInterface.cpp
@@ -753,107 +753,360 @@ void FChaosEngineInterface::SetAngularDamping_AssumesLocked(const FPhysicsActorH
 	}
 }
 
-void FChaosEngineInterface::AddImpulse_AssumesLocked(const FPhysicsActorHandle& InActorReference,const FVector& InForce)
+template<typename BodyHandleType>
+struct FChaosStateOps
+{
+	static void AddImpulse(BodyHandleType& BodyHandle, const FVector& InForce)
+	{
+		BodyHandle.SetLinearImpulse(BodyHandle.LinearImpulse() + InForce, /*bIsVelocity=*/false);
+	}
+
+	static void AddAngularImpulseInRadians(BodyHandleType& BodyHandle, const FVector& InTorque)
+	{
+		BodyHandle.SetAngularImpulse(BodyHandle.AngularImpulse() + InTorque, /*bIsVelocity=*/false);
+	}
+
+	static void AddVelocity(BodyHandleType& BodyHandle, const FVector& InVelocityDelta)
+	{
+		AddImpulse(BodyHandle, BodyHandle.M() * InVelocityDelta);
+	}
+
+	static void AddAngularVelocityInRadians(BodyHandleType& BodyHandle, const FVector& InAngularVelocityDeltaRad)
+	{
+		const Chaos::FMatrix33 WorldI = Chaos::FParticleUtilitiesXR::GetWorldInertia(&BodyHandle);
+		AddAngularImpulseInRadians(BodyHandle, WorldI * InAngularVelocityDeltaRad);
+	}
+
+	static void AddImpulseAtLocation(BodyHandleType& BodyHandle, const FVector& InImpulse, const FVector& InLocation)
+	{
+		const Chaos::FVec3 WorldCOM = Chaos::FParticleUtilitiesGT::GetCoMWorldPosition(&BodyHandle);
+		const Chaos::FVec3 AngularImpulse = Chaos::FVec3::CrossProduct(InLocation - WorldCOM, InImpulse);
+		AddImpulse(BodyHandle, InImpulse);
+		AddAngularImpulseInRadians(BodyHandle, AngularImpulse);
+	}
+
+	static void AddVelocityChangeImpulseAtLocation(BodyHandleType& BodyHandle, const FVector& InVelocityDelta, const FVector& InLocation)
+	{
+		AddImpulseAtLocation(BodyHandle, BodyHandle.M() * InVelocityDelta, InLocation);
+	}
+
+	static void AddRadialImpulse(BodyHandleType& BodyHandle, const FVector& InOrigin, float InRadius, float InStrength, ERadialImpulseFalloff InFalloff, bool bInVelChange)
+	{
+		const Chaos::FVec3 WorldCOM = Chaos::FParticleUtilitiesGT::GetCoMWorldPosition(&BodyHandle);
+		const Chaos::FVec3 OriginToActor = WorldCOM - InOrigin;
+		const Chaos::FReal OriginToActorDistance = OriginToActor.Size();
+		if (OriginToActorDistance < InRadius)
+		{
+			Chaos::FVec3 FinalImpulse = FVector::ZeroVector;
+			if (OriginToActorDistance > 0)
+			{
+				const Chaos::FVec3 OriginToActorNorm = OriginToActor / OriginToActorDistance;
+
+				if (InFalloff == ERadialImpulseFalloff::RIF_Constant)
+				{
+					FinalImpulse = OriginToActorNorm * InStrength;
+				}
+				else if (InFalloff == ERadialImpulseFalloff::RIF_Linear)
+				{
+					const Chaos::FReal DistanceOverlapping = InRadius - OriginToActorDistance;
+					if (DistanceOverlapping > 0)
+					{
+						FinalImpulse = OriginToActorNorm * FMath::Lerp(0.0f, InStrength, DistanceOverlapping / InRadius);
+					}
+				}
+				else
+				{
+					// Unimplemented falloff type
+					ensure(false);
+				}
+			}
+			else
+			{
+				// Sphere and actor center are coincident, just pick a direction and apply maximum strength impulse.
+				FinalImpulse = FVector::ForwardVector * InStrength;
+			}
+
+			if (bInVelChange)
+			{
+				AddVelocity(BodyHandle, FinalImpulse);
+			}
+			else
+			{
+				AddImpulse(BodyHandle, FinalImpulse);
+			}
+		}
+	}
+
+	static void AddForce(BodyHandleType& BodyHandle, const FVector& Force, bool bAllowSubstepping, bool bAccelChange)
+	{
+		Chaos::EObjectStateType ObjectState = BodyHandle.ObjectState();
+		BodyHandle.SetObjectState(Chaos::EObjectStateType::Dynamic);
+
+		if (bAccelChange)
+		{
+			const Chaos::FReal Mass = BodyHandle.M();
+			const Chaos::FVec3 Acceleration = Force * Mass;
+			BodyHandle.AddForce(Acceleration);
+		}
+		else
+		{
+			BodyHandle.AddForce(Force);
+		}
+	}
+
+	static void AddForceAtPosition(BodyHandleType& BodyHandle, const FVector& Force, const FVector& Position, bool bAllowSubstepping, bool bIsLocalForce /*= false*/)
+	{
+		
+		Chaos::EObjectStateType ObjectState = BodyHandle.ObjectState();
+		const Chaos::FVec3 WorldCOM = Chaos::FParticleUtilitiesGT::GetCoMWorldPosition(&BodyHandle);
+
+		BodyHandle.SetObjectState(Chaos::EObjectStateType::Dynamic);
+
+		if (bIsLocalForce)
+		{
+			const Chaos::FRigidTransform3 CurrentTransform = Chaos::FParticleUtilitiesGT::GetActorWorldTransform(&BodyHandle);
+			const Chaos::FVec3 WorldPosition = CurrentTransform.TransformPosition(Position);
+			const Chaos::FVec3 WorldForce = CurrentTransform.TransformVector(Force);
+			const Chaos::FVec3 WorldTorque = Chaos::FVec3::CrossProduct(WorldPosition - WorldCOM, WorldForce);
+			BodyHandle.AddForce(WorldForce);
+			BodyHandle.AddTorque(WorldTorque);
+		}
+		else
+		{
+			const Chaos::FVec3 WorldTorque = Chaos::FVec3::CrossProduct(Position - WorldCOM, Force);
+			BodyHandle.AddForce(Force);
+			BodyHandle.AddTorque(WorldTorque);
+		}
+	}
+
+	static void AddRadialForce(BodyHandleType& BodyHandle, const FVector& Origin, const float Radius, const float Strength, const uint8 Falloff, bool bAccelChange, bool bAllowSubstepping)
+	{
+		Chaos::EObjectStateType ObjectState = BodyHandle.ObjectState();
+		if (CHAOS_ENSURE(ObjectState == Chaos::EObjectStateType::Dynamic || ObjectState == Chaos::EObjectStateType::Sleeping))
+		{
+			const Chaos::FVec3 WorldCOM = Chaos::FParticleUtilitiesGT::GetCoMWorldPosition(&BodyHandle);
+
+			Chaos::FVec3 Direction = WorldCOM - Origin;
+			const Chaos::FReal Distance = Direction.Size();
+			if (Distance > Radius)
+			{
+				return;
+			}
+
+			BodyHandle.SetObjectState(Chaos::EObjectStateType::Dynamic);
+
+			if (Distance < 1e-4)
+			{
+				Direction = Chaos::FVec3(1, 0, 0);
+			}
+			else
+			{
+				Direction = Direction.GetUnsafeNormal();
+			}
+			Chaos::FVec3 Force(0, 0, 0);
+			CHAOS_ENSURE(Falloff < RIF_MAX);
+			if (Falloff == ERadialImpulseFalloff::RIF_Constant)
+			{
+				Force = Strength * Direction;
+			}
+			if (Falloff == ERadialImpulseFalloff::RIF_Linear)
+			{
+				Force = (Radius - Distance) / Radius * Strength * Direction;
+			}
+			if (bAccelChange)
+			{
+				const Chaos::FReal Mass = BodyHandle.M();
+				const Chaos::FVec3 Acceleration = Force * Mass;
+				BodyHandle.AddForce(Acceleration);
+			}
+			else
+			{
+				BodyHandle.AddForce(Force);
+			}
+		}
+	}
+
+	static void AddTorque(BodyHandleType& BodyHandle, const FVector& Torque, bool bAllowSubstepping, bool bAccelChange)
+	{
+		Chaos::EObjectStateType ObjectState = BodyHandle.ObjectState();
+		if (CHAOS_ENSURE(ObjectState == Chaos::EObjectStateType::Dynamic || ObjectState == Chaos::EObjectStateType::Sleeping))
+		{
+			if (bAccelChange)
+			{
+				BodyHandle.AddTorque(Chaos::FParticleUtilitiesXR::GetWorldInertia(&BodyHandle) * Torque);
+			}
+			else
+			{
+				BodyHandle.AddTorque(Torque);
+			}
+		}
+	}
+};
+
+void FChaosEngineInterface::AddImpulse_AssumesLocked(const FPhysicsActorHandle& InActorReference,const FVector& InForce, bool bIsInternal)
 {
 	if(ensure(FChaosEngineInterface::IsValid(InActorReference)))
 	{
-		Chaos::FRigidBodyHandle_External& Body_External = InActorReference->GetGameThreadAPI();
-		Body_External.SetLinearImpulse(Body_External.LinearImpulse() + InForce, /*bIsVelocity=*/false);
+		if (bIsInternal)
+		{
+			FChaosStateOps<Chaos::FRigidBodyHandle_Internal>::AddImpulse(*InActorReference->GetPhysicsThreadAPI(), InForce);
+		}
+		else
+		{
+			FChaosStateOps<Chaos::FRigidBodyHandle_External>::AddImpulse(InActorReference->GetGameThreadAPI(), InForce);
+		}
 	}
 }
 
-void FChaosEngineInterface::AddAngularImpulseInRadians_AssumesLocked(const FPhysicsActorHandle& InActorReference,const FVector& InTorque)
+void FChaosEngineInterface::AddAngularImpulseInRadians_AssumesLocked(const FPhysicsActorHandle& InActorReference,const FVector& InTorque, bool bIsInternal)
 {
 	if(ensure(FChaosEngineInterface::IsValid(InActorReference)))
 	{
-		Chaos::FRigidBodyHandle_External& Body_External = InActorReference->GetGameThreadAPI();
-		Body_External.SetAngularImpulse(Body_External.AngularImpulse() + InTorque, /*bIsVelocity=*/false);
+		if (bIsInternal)
+		{
+			FChaosStateOps<Chaos::FRigidBodyHandle_Internal>::AddAngularImpulseInRadians(*InActorReference->GetPhysicsThreadAPI(), InTorque);
+		}
+		else
+		{
+			FChaosStateOps<Chaos::FRigidBodyHandle_External>::AddAngularImpulseInRadians(InActorReference->GetGameThreadAPI(), InTorque);
+		}
 	}
 }
 
-void FChaosEngineInterface::AddVelocity_AssumesLocked(const FPhysicsActorHandle& InActorReference,const FVector& InVelocityDelta)
+void FChaosEngineInterface::AddVelocity_AssumesLocked(const FPhysicsActorHandle& InActorReference,const FVector& InVelocityDelta, bool bIsInternal)
 {
-	if(ensure(FChaosEngineInterface::IsValid(InActorReference)))
+	if (ensure(FChaosEngineInterface::IsValid(InActorReference)))
 	{
-		AddImpulse_AssumesLocked(InActorReference, InActorReference->GetGameThreadAPI().M() * InVelocityDelta);
+		if (bIsInternal)
+		{
+			FChaosStateOps<Chaos::FRigidBodyHandle_Internal>::AddVelocity(*InActorReference->GetPhysicsThreadAPI(), InVelocityDelta);
+		}
+		else
+		{
+			FChaosStateOps<Chaos::FRigidBodyHandle_External>::AddVelocity(InActorReference->GetGameThreadAPI(), InVelocityDelta);
+		}
 	}
 }
 
-void FChaosEngineInterface::AddAngularVelocityInRadians_AssumesLocked(const FPhysicsActorHandle& InActorReference,const FVector& InAngularVelocityDeltaRad)
+void FChaosEngineInterface::AddAngularVelocityInRadians_AssumesLocked(const FPhysicsActorHandle& InActorReference,const FVector& InAngularVelocityDeltaRad, bool bIsInternal)
 {
-	if(ensure(FChaosEngineInterface::IsValid(InActorReference)))
+	if (ensure(FChaosEngineInterface::IsValid(InActorReference)))
 	{
-		const Chaos::FMatrix33 WorldI = Chaos::FParticleUtilitiesXR::GetWorldInertia(&InActorReference->GetGameThreadAPI());
-		AddAngularImpulseInRadians_AssumesLocked(InActorReference,WorldI * InAngularVelocityDeltaRad);
+		if (bIsInternal)
+		{
+			FChaosStateOps<Chaos::FRigidBodyHandle_Internal>::AddAngularVelocityInRadians(*InActorReference->GetPhysicsThreadAPI(), InAngularVelocityDeltaRad);
+		}
+		else
+		{
+			FChaosStateOps<Chaos::FRigidBodyHandle_External>::AddAngularVelocityInRadians(InActorReference->GetGameThreadAPI(), InAngularVelocityDeltaRad);
+		}
 	}
 }
 
-void FChaosEngineInterface::AddImpulseAtLocation_AssumesLocked(const FPhysicsActorHandle& InActorReference,const FVector& InImpulse,const FVector& InLocation)
+void FChaosEngineInterface::AddImpulseAtLocation_AssumesLocked(const FPhysicsActorHandle& InActorReference,const FVector& InImpulse,const FVector& InLocation, bool bIsInternal)
 {
-	if(ensure(FChaosEngineInterface::IsValid(InActorReference)))
+	if (ensure(FChaosEngineInterface::IsValid(InActorReference)))
 	{
-		const Chaos::FVec3 WorldCOM = Chaos::FParticleUtilitiesGT::GetCoMWorldPosition(&InActorReference->GetGameThreadAPI());
-		const Chaos::FVec3 AngularImpulse = Chaos::FVec3::CrossProduct(InLocation - WorldCOM,InImpulse);
-		AddImpulse_AssumesLocked(InActorReference,InImpulse);
-		AddAngularImpulseInRadians_AssumesLocked(InActorReference,AngularImpulse);
+		if (bIsInternal)
+		{
+			
+			FChaosStateOps<Chaos::FRigidBodyHandle_Internal>::AddImpulseAtLocation(*InActorReference->GetPhysicsThreadAPI(), InImpulse, InLocation);
+			//UE_LOG(LogTemp, Warning, TEXT("AddImpulseAtLocation_AssumesLocked Impulse = %s | State = %d"), *InImpulse.ToString(), InActorReference->GetPhysicsThreadAPI()->ObjectState());
+		}
+		else
+		{
+			FChaosStateOps<Chaos::FRigidBodyHandle_External>::AddImpulseAtLocation(InActorReference->GetGameThreadAPI(), InImpulse, InLocation);
+		}
 	}
 }
 
-void FChaosEngineInterface::AddVelocityChangeImpulseAtLocation_AssumesLocked(const FPhysicsActorHandle& InActorReference, const FVector& InVelocityDelta, const FVector& InLocation)
+void FChaosEngineInterface::AddVelocityChangeImpulseAtLocation_AssumesLocked(const FPhysicsActorHandle& InActorReference, const FVector& InVelocityDelta, const FVector& InLocation, bool bIsInternal)
 {
 	if (ensure(FChaosEngineInterface::IsValid(InActorReference)))
 	{
-		AddImpulseAtLocation_AssumesLocked(InActorReference, InActorReference->GetGameThreadAPI().M() * InVelocityDelta, InLocation);
+		if (bIsInternal)
+		{
+			FChaosStateOps<Chaos::FRigidBodyHandle_Internal>::AddVelocityChangeImpulseAtLocation(*InActorReference->GetPhysicsThreadAPI(), InVelocityDelta, InLocation);
+		}
+		else
+		{
+			FChaosStateOps<Chaos::FRigidBodyHandle_External>::AddVelocityChangeImpulseAtLocation(InActorReference->GetGameThreadAPI(), InVelocityDelta, InLocation);
+		}
 	}
 }
 
-void FChaosEngineInterface::AddRadialImpulse_AssumesLocked(const FPhysicsActorHandle& InActorReference,const FVector& InOrigin,float InRadius,float InStrength,ERadialImpulseFalloff InFalloff,bool bInVelChange)
+void FChaosEngineInterface::AddRadialImpulse_AssumesLocked(const FPhysicsActorHandle& InActorReference,const FVector& InOrigin,float InRadius,float InStrength,ERadialImpulseFalloff InFalloff,bool bInVelChange, bool bIsInternal)
 {
-	using namespace Chaos;
-	if (ensure(InActorReference->GetGameThreadAPI().CanTreatAsRigid()))
+	if (ensure(FChaosEngineInterface::IsValid(InActorReference)) && ensure(InActorReference->GetGameThreadAPI().CanTreatAsRigid()))
 	{
-		const FVec3 WorldCOM = FParticleUtilitiesGT::GetCoMWorldPosition(&InActorReference->GetGameThreadAPI());
-		const FVec3 OriginToActor = WorldCOM - InOrigin;
-		const FReal OriginToActorDistance = OriginToActor.Size();
-		if(OriginToActorDistance < InRadius)
+		if (bIsInternal)
 		{
-			FVec3 FinalImpulse = FVector::ZeroVector;
-			if(OriginToActorDistance > 0)
+			FChaosStateOps<Chaos::FRigidBodyHandle_Internal>::AddRadialImpulse(*InActorReference->GetPhysicsThreadAPI(), InOrigin, InRadius, InStrength, InFalloff, bInVelChange);
+		}
+		else
 		{
-			const FVec3 OriginToActorNorm = OriginToActor / OriginToActorDistance;
+			FChaosStateOps<Chaos::FRigidBodyHandle_External>::AddRadialImpulse(InActorReference->GetGameThreadAPI(), InOrigin, InRadius, InStrength, InFalloff, bInVelChange);
+		}
+	}
+}
 
-				if(InFalloff == ERadialImpulseFalloff::RIF_Constant)
-			{
-					FinalImpulse = OriginToActorNorm * InStrength;
-			}
-				else if(InFalloff == ERadialImpulseFalloff::RIF_Linear)
-			{
-				const FReal DistanceOverlapping = InRadius - OriginToActorDistance;
-					if(DistanceOverlapping > 0)
-				{
-						FinalImpulse = OriginToActorNorm * FMath::Lerp(0.0f, InStrength, DistanceOverlapping / InRadius);
-				}
-			}
-			else
-			{
-				// Unimplemented falloff type
-				ensure(false);
-			}
+void FChaosEngineInterface::AddForce_AssumesLocked(const FPhysicsActorHandle& InActorReference, const FVector& Force, bool bAllowSubstepping, bool bAccelChange, bool bIsInternal)
+{
+	if (ensure(FChaosEngineInterface::IsValid(InActorReference)))
+	{
+		if (bIsInternal)
+		{
+			FChaosStateOps<Chaos::FRigidBodyHandle_Internal>::AddForce(*InActorReference->GetPhysicsThreadAPI(), Force, bAllowSubstepping, bAccelChange);
 		}
 		else
 		{
-			// Sphere and actor center are coincident, just pick a direction and apply maximum strength impulse.
-				FinalImpulse = FVector::ForwardVector * InStrength;
-			}
+			FChaosStateOps<Chaos::FRigidBodyHandle_External>::AddForce(InActorReference->GetGameThreadAPI(), Force, bAllowSubstepping, bAccelChange);
+		}
+	}
+}
 
-			if(bInVelChange)
-			{
-				AddVelocity_AssumesLocked(InActorReference, FinalImpulse);
-			}
-			else
-			{
-				AddImpulse_AssumesLocked(InActorReference, FinalImpulse);
-			}
+void FChaosEngineInterface::AddForceAtPosition_AssumesLocked(const FPhysicsActorHandle& InActorReference, const FVector& Force, const FVector& Position, bool bAllowSubstepping, bool bIsLocalForce, bool bIsInternal)
+{
+	if (ensure(FChaosEngineInterface::IsValid(InActorReference)))
+	{
+		if (bIsInternal)
+		{
+			FChaosStateOps<Chaos::FRigidBodyHandle_Internal>::AddForceAtPosition(*InActorReference->GetPhysicsThreadAPI(), Force, Position, bAllowSubstepping, bIsLocalForce);
+			//UE_LOG(LogTemp, Warning, TEXT("AddForceAtPosition_AssumesLocked Force = %s | State = %d"), *Force.ToString(), InActorReference->GetPhysicsThreadAPI()->ObjectState());
+		}
+		else
+		{
+			FChaosStateOps<Chaos::FRigidBodyHandle_External>::AddForceAtPosition(InActorReference->GetGameThreadAPI(), Force, Position,  bAllowSubstepping, bIsLocalForce);
+		}
+	}
+}
+
+void FChaosEngineInterface::AddRadialForce_AssumesLocked(const FPhysicsActorHandle& InActorReference, const FVector& Origin, const float Radius, const float Strength, const uint8 Falloff, bool bAccelChange, bool bAllowSubstepping, bool bIsInternal)
+{
+	if (ensure(FChaosEngineInterface::IsValid(InActorReference)))
+	{
+		if (bIsInternal)
+		{
+			FChaosStateOps<Chaos::FRigidBodyHandle_Internal>::AddRadialForce(*InActorReference->GetPhysicsThreadAPI(), Origin, Radius, Strength, Falloff, bAccelChange, bAllowSubstepping);
+		}
+		else
+		{
+			FChaosStateOps<Chaos::FRigidBodyHandle_External>::AddRadialForce(InActorReference->GetGameThreadAPI(), Origin, Radius, Strength, Falloff, bAccelChange, bAllowSubstepping);
+		}
+	}
+}
+
+void FChaosEngineInterface::AddTorque_AssumesLocked(const FPhysicsActorHandle& InActorReference, const FVector& Torque, bool bAllowSubstepping, bool bAccelChange, bool bIsInternal)
+{
+	if (ensure(FChaosEngineInterface::IsValid(InActorReference)))
+	{
+		if (bIsInternal)
+		{
+			FChaosStateOps<Chaos::FRigidBodyHandle_Internal>::AddTorque(*InActorReference->GetPhysicsThreadAPI(), Torque, bAllowSubstepping, bAccelChange);
+		}
+		else
+		{
+			FChaosStateOps<Chaos::FRigidBodyHandle_External>::AddTorque(InActorReference->GetGameThreadAPI(), Torque, bAllowSubstepping, bAccelChange);
 		}
 	}
 }
diff --git a/Engine/Source/Runtime/PhysicsCore/Public/Chaos/ChaosEngineInterface.h b/Engine/Source/Runtime/PhysicsCore/Public/Chaos/ChaosEngineInterface.h
--- a/Engine/Source/Runtime/PhysicsCore/Public/Chaos/ChaosEngineInterface.h
+++ b/Engine/Source/Runtime/PhysicsCore/Public/Chaos/ChaosEngineInterface.h
@@ -443,13 +443,18 @@ public:
 	static void SetLinearDamping_AssumesLocked(const FPhysicsActorHandle& InActorReference,float InDamping);
 	static void SetAngularDamping_AssumesLocked(const FPhysicsActorHandle& InActorReference,float InDamping);
 
-	static void AddImpulse_AssumesLocked(const FPhysicsActorHandle& InActorReference,const FVector& InForce);
-	static void AddAngularImpulseInRadians_AssumesLocked(const FPhysicsActorHandle& InActorReference,const FVector& InTorque);
-	static void AddVelocity_AssumesLocked(const FPhysicsActorHandle& InActorReference,const FVector& InForce);
-	static void AddAngularVelocityInRadians_AssumesLocked(const FPhysicsActorHandle& InActorReference,const FVector& InTorque);
-	static void AddImpulseAtLocation_AssumesLocked(const FPhysicsActorHandle& InActorReference,const FVector& InImpulse,const FVector& InLocation);
-	static void AddVelocityChangeImpulseAtLocation_AssumesLocked(const FPhysicsActorHandle& InActorReference, const FVector& InVelocityDelta, const FVector& InLocation);
-	static void AddRadialImpulse_AssumesLocked(const FPhysicsActorHandle& InActorReference,const FVector& InOrigin,float InRadius,float InStrength,ERadialImpulseFalloff InFalloff,bool bInVelChange);
+	static void AddForce_AssumesLocked(const FPhysicsActorHandle& InActorReference, const FVector& Force, bool bAllowSubstepping, bool bAccelChange, bool bIsInternal = false);
+	static void AddForceAtPosition_AssumesLocked(const FPhysicsActorHandle& InActorReference, const FVector& Force, const FVector& Position, bool bAllowSubstepping, bool bIsLocalForce = false, bool bIsInternal = false);
+	static void AddRadialForce_AssumesLocked(const FPhysicsActorHandle& InActorReference, const FVector& Origin, const float Radius, const float Strength, const uint8 Falloff, bool bAccelChange, bool bAllowSubstepping, bool bIsInternal = false);
+	static void AddTorque_AssumesLocked(const FPhysicsActorHandle& InActorReference, const FVector& Torque, bool bAllowSubstepping, bool bAccelChange, bool bIsInternal = false);
+
+	static void AddImpulse_AssumesLocked(const FPhysicsActorHandle& InActorReference,const FVector& InForce, bool bIsInternal = false);
+	static void AddAngularImpulseInRadians_AssumesLocked(const FPhysicsActorHandle& InActorReference,const FVector& InTorque, bool bIsInternal = false);
+	static void AddVelocity_AssumesLocked(const FPhysicsActorHandle& InActorReference,const FVector& InForce, bool bIsInternal = false);
+	static void AddAngularVelocityInRadians_AssumesLocked(const FPhysicsActorHandle& InActorReference,const FVector& InTorque, bool bIsInternal = false);
+	static void AddImpulseAtLocation_AssumesLocked(const FPhysicsActorHandle& InActorReference,const FVector& InImpulse,const FVector& InLocation, bool bIsInternal = false);
+	static void AddVelocityChangeImpulseAtLocation_AssumesLocked(const FPhysicsActorHandle& InActorReference, const FVector& InVelocityDelta, const FVector& InLocation, bool bIsInternal = false);
+	static void AddRadialImpulse_AssumesLocked(const FPhysicsActorHandle& InActorReference,const FVector& InOrigin,float InRadius,float InStrength,ERadialImpulseFalloff InFalloff,bool bInVelChange, bool bIsInternal = false);
 
 	static bool IsGravityEnabled_AssumesLocked(const FPhysicsActorHandle& InActorReference);
 	static void SetGravityEnabled_AssumesLocked(const FPhysicsActorHandle& InActorReference,bool bEnabled);
