Static cvar query for networked physics

diff --git a/Engine/Plugins/Experimental/ChaosVehiclesPlugin/Source/ChaosVehicles/Private/ChaosVehicleMovementComponent.cpp b/Engine/Plugins/Experimental/ChaosVehiclesPlugin/Source/ChaosVehicles/Private/ChaosVehicleMovementComponent.cpp
--- a/Engine/Plugins/Experimental/ChaosVehiclesPlugin/Source/ChaosVehicles/Private/ChaosVehicleMovementComponent.cpp
+++ b/Engine/Plugins/Experimental/ChaosVehiclesPlugin/Source/ChaosVehicles/Private/ChaosVehicleMovementComponent.cpp
@@ -1068,7 +1068,7 @@ void UChaosVehicleMovementComponent::OnCreatePhysicsState()
 							AController* Controller = GetController();
 							const bool bIsLocalHistory = (Controller && Controller->IsLocalController());
 
-							const IConsoleVariable* NumFramesCVar = IConsoleManager::Get().FindConsoleVariable(TEXT("p.RewindCaptureNumFrames"));
+							static const IConsoleVariable* NumFramesCVar = IConsoleManager::Get().FindConsoleVariable(TEXT("p.RewindCaptureNumFrames"));
 							const int32 NumFrames = NumFramesCVar ? NumFramesCVar->GetInt() : 50;
 
 							VehicleSimulationPT->InputsHistory = MakeShared<FVehicleInputsHistory>(NumFrames, bIsLocalHistory, VehicleSimulationPT.Get());
diff --git a/Engine/Source/Runtime/Engine/Private/PhysicsEngine/AsyncPhysicsInputComponent.cpp b/Engine/Source/Runtime/Engine/Private/PhysicsEngine/AsyncPhysicsInputComponent.cpp
--- a/Engine/Source/Runtime/Engine/Private/PhysicsEngine/AsyncPhysicsInputComponent.cpp
+++ b/Engine/Source/Runtime/Engine/Private/PhysicsEngine/AsyncPhysicsInputComponent.cpp
@@ -343,7 +343,7 @@ void UAsyncPhysicsInputComponent::SetDataClass(TSubclassOf<UAsyncPhysicsData> In
 				// -------------------------------------------------------------------------------------
 
 				int32 NumFrames = 64;
-				const IConsoleVariable* NumFramesCVar = IConsoleManager::Get().FindConsoleVariable(TEXT("p.RewindCaptureNumFrames"));
+				static const IConsoleVariable* NumFramesCVar = IConsoleManager::Get().FindConsoleVariable(TEXT("p.RewindCaptureNumFrames"));
 				if (ensure(NumFramesCVar))
 				{
 					// 1 frame is required to enable rewind capture. 
diff --git a/Engine/Source/Runtime/Engine/Private/PhysicsEngine/Experimental/PhysScene_Chaos.cpp b/Engine/Source/Runtime/Engine/Private/PhysicsEngine/Experimental/PhysScene_Chaos.cpp
--- a/Engine/Source/Runtime/Engine/Private/PhysicsEngine/Experimental/PhysScene_Chaos.cpp
+++ b/Engine/Source/Runtime/Engine/Private/PhysicsEngine/Experimental/PhysScene_Chaos.cpp
@@ -97,7 +97,7 @@ public:
 		static IConsoleVariable* ResimNetworkPhysicsCVar = IConsoleManager::Get().FindConsoleVariable(TEXT("np2.ResimNetworkPhysicsPrediction"));
 		const bool bEnableResim =  (ResimNetworkPhysicsCVar && ResimNetworkPhysicsCVar->GetInt() == 1);
 
-		const IConsoleVariable* NumFramesCVar = IConsoleManager::Get().FindConsoleVariable(TEXT("p.RewindCaptureNumFrames"));
+		static const IConsoleVariable* NumFramesCVar = IConsoleManager::Get().FindConsoleVariable(TEXT("p.RewindCaptureNumFrames"));
 		const int32 NumFrames = NumFramesCVar ? NumFramesCVar->GetInt() : 0;
 
 		for (int32 Idx = 0; Idx < PendingCommands.Num(); ++Idx)
