TP_VehicleAdvBP, fix the offroad vehicle suspension where vehicle wheels are sitting through the ground.
The default case for the wheel trace type was the wrong way round (different default in some parts of code than before).

Recent'ish changes for Opal introduced this, but doesn't affect Opal.

diff --git a/Engine/Plugins/Experimental/ChaosVehiclesPlugin/Source/ChaosVehicles/Private/ChaosWheeledVehicleMovementComponent.cpp b/Engine/Plugins/Experimental/ChaosVehiclesPlugin/Source/ChaosVehicles/Private/ChaosWheeledVehicleMovementComponent.cpp
--- a/Engine/Plugins/Experimental/ChaosVehiclesPlugin/Source/ChaosVehicles/Private/ChaosWheeledVehicleMovementComponent.cpp
+++ b/Engine/Plugins/Experimental/ChaosVehiclesPlugin/Source/ChaosVehicles/Private/ChaosWheeledVehicleMovementComponent.cpp
@@ -615,13 +615,13 @@ void UChaosWheeledVehicleSimulation::ApplySuspensionForces(float DeltaTime, TArr
 						if (FSuspensionConstraintPhysicsProxy* Proxy = Constraint->GetProxy<FSuspensionConstraintPhysicsProxy>())
 						{
 							FVec3 TargetPos;
-							if (WheelTraceParams[WheelIdx].SweepShape == ESweepShape::Raycast)
+							if (WheelTraceParams[WheelIdx].SweepShape == ESweepShape::Spherecast)
 							{
-								TargetPos = HitResult.ImpactPoint + (PWheel.GetEffectiveRadius() * VehicleState.VehicleUpAxis);
+								TargetPos = HitResult.Location;
 							}
 							else
 							{
-								TargetPos = HitResult.Location;
+								TargetPos = HitResult.ImpactPoint + (PWheel.GetEffectiveRadius() * VehicleState.VehicleUpAxis);
 							}
 
 							Chaos::FPhysicsSolver* Solver = Proxy->GetSolver<Chaos::FPhysicsSolver>();
@@ -1241,7 +1241,6 @@ void UChaosWheeledVehicleMovementComponent::FixupSkeletalMesh()
 				}
 			}
 		}
-
 		VehicleSimulationPT->UpdateConstraintHandles(ConstraintHandles); // TODO: think of a better way to communicate this data
 
 		Mesh->KinematicBonesUpdateType = EKinematicBonesUpdateToPhysics::SkipSimulatingBones;
@@ -2891,14 +2890,14 @@ float UChaosWheeledVehicleMovementComponent::GetSuspensionOffset(int WheelIndex)
 					if (HitResult.bBlockingHit)
 					{
 						FVector LocalPos = GetWheelRestingPosition(WheelSetup);
-						if (Wheel->SweepShape == ESweepShape::Raycast)
+						if (Wheel->SweepShape == ESweepShape::Spherecast)
 						{
-							FVector LocalHitPoint = VehicleWorldTransform.InverseTransformPosition(HitResult.ImpactPoint);
+							FVector LocalHitPoint = VehicleWorldTransform.InverseTransformPosition(HitResult.Location);
 							Offset = CalcOffset(Wheel->SweepShape, LocalHitPoint, LocalPos, Wheel->WheelRadius);
 						}
 						else
 						{
-							FVector LocalHitPoint = VehicleWorldTransform.InverseTransformPosition(HitResult.Location);
+							FVector LocalHitPoint = VehicleWorldTransform.InverseTransformPosition(HitResult.ImpactPoint);
 							Offset = CalcOffset(Wheel->SweepShape, LocalHitPoint, LocalPos, Wheel->WheelRadius);
 						}
 
