Fix CIS uninitialized data

diff --git a/Engine/Plugins/Experimental/ChaosVehiclesPlugin/Source/ChaosVehicles/Private/ChaosVehicleMovementComponent.cpp b/Engine/Plugins/Experimental/ChaosVehiclesPlugin/Source/ChaosVehicles/Private/ChaosVehicleMovementComponent.cpp
--- a/Engine/Plugins/Experimental/ChaosVehiclesPlugin/Source/ChaosVehicles/Private/ChaosVehicleMovementComponent.cpp
+++ b/Engine/Plugins/Experimental/ChaosVehiclesPlugin/Source/ChaosVehicles/Private/ChaosVehicleMovementComponent.cpp
@@ -74,6 +74,19 @@ FAutoConsoleVariableRef CVarChaosVehiclesControlInputWakeTolerance(TEXT("p.Vehic
 
 using namespace Chaos;
 
+FVehicleInputsDatas::FVehicleInputsDatas()
+	: TransmissionChangeTime(0.0f)
+	, TransmissionCurrentGear(0)
+	, TransmissionTargetGear(0)
+{
+}
+
+FVehicleStatesDatas::FVehicleStatesDatas()
+	: StateLastVelocity()
+	, EngineOmega()
+{
+}
+
 void FVehicleState::CaptureState(const FBodyInstance* TargetInstance, float GravityZ, float DeltaTime)
 {
 	if (TargetInstance)
diff --git a/Engine/Plugins/Experimental/ChaosVehiclesPlugin/Source/ChaosVehicles/Public/ChaosVehicleMovementComponent.h b/Engine/Plugins/Experimental/ChaosVehiclesPlugin/Source/ChaosVehicles/Public/ChaosVehicleMovementComponent.h
--- a/Engine/Plugins/Experimental/ChaosVehiclesPlugin/Source/ChaosVehicles/Public/ChaosVehicleMovementComponent.h
+++ b/Engine/Plugins/Experimental/ChaosVehiclesPlugin/Source/ChaosVehicles/Public/ChaosVehicleMovementComponent.h
@@ -124,6 +124,8 @@ struct CHAOSVEHICLES_API FVehicleInputsDatas : public FVehicleHistoryDatas
 {
 	GENERATED_USTRUCT_BODY()
 
+	FVehicleInputsDatas();
+
 	/** List of incoming control inputs coming from the local client */
 	UPROPERTY()
 	FVehicleInputs VehicleInputs;
@@ -147,6 +149,8 @@ struct CHAOSVEHICLES_API FVehicleStatesDatas : public FVehicleHistoryDatas
 {
 	GENERATED_USTRUCT_BODY()
 
+	FVehicleStatesDatas();
+
 	/** Vehicle state last velocity */
 	UPROPERTY()
 	FVector StateLastVelocity;
