Fix CIS warnings

diff --git a/Engine/Plugins/Experimental/ChaosVehiclesPlugin/Source/ChaosVehicles/Private/ChaosVehicleMovementComponent.cpp b/Engine/Plugins/Experimental/ChaosVehiclesPlugin/Source/ChaosVehicles/Private/ChaosVehicleMovementComponent.cpp
--- a/Engine/Plugins/Experimental/ChaosVehiclesPlugin/Source/ChaosVehicles/Private/ChaosVehicleMovementComponent.cpp
+++ b/Engine/Plugins/Experimental/ChaosVehiclesPlugin/Source/ChaosVehicles/Private/ChaosVehicleMovementComponent.cpp
@@ -72,7 +72,6 @@ FAutoConsoleVariableRef CVarChaosVehiclesSetMaxMPH(TEXT("p.Vehicle.SetMaxMPH"),
 FAutoConsoleVariableRef CVarChaosVehiclesEnableMultithreading(TEXT("p.Vehicle.EnableMultithreading"), GVehicleDebugParams.EnableMultithreading, TEXT("Enable multi-threading of vehicle updates."));
 FAutoConsoleVariableRef CVarChaosVehiclesControlInputWakeTolerance(TEXT("p.Vehicle.ControlInputWakeTolerance"), GVehicleDebugParams.ControlInputWakeTolerance, TEXT("Set the control input wake tolerance."));
 
-using namespace Chaos;
 
 FVehicleInputsDatas::FVehicleInputsDatas()
 	: TransmissionChangeTime(0.0f)
@@ -82,8 +81,8 @@ FVehicleInputsDatas::FVehicleInputsDatas()
 }
 
 FVehicleStatesDatas::FVehicleStatesDatas()
-	: StateLastVelocity()
-	, EngineOmega()
+	: StateLastVelocity(FVector::ZeroVector)
+	, EngineOmega(0.0f)
 {
 }
 
@@ -223,6 +222,8 @@ bool FVehicleStatesHistory::RecordStates(const int32 RecordFrame, const DatasTyp
 
 void FVehicleStatesHistory::ReplicateStates(DatasType& ReplicatedDatas)
 {
+	using namespace Chaos;
+
 	if (VehicleSimulation)
 	{
 		if (TUniquePtr<Chaos::FSimpleWheeledVehicle>& Vehicle = VehicleSimulation->PVehicle)
@@ -270,6 +271,8 @@ void FVehicleStatesHistory::ReplicateStates(DatasType& ReplicatedDatas)
 
 bool FVehicleStatesHistory::RewindStates(const int32 RewindFrame, const bool bResetSolver)
 {
+	using namespace Chaos;
+
 	if (VehicleSimulation && Super::RewindStates(RewindFrame, bResetSolver))
 	{
 		const FVehicleStatesDatas& CurrentDatas = GetCurrentDatas();
@@ -352,6 +355,8 @@ void FVehicleInputsHistory::ReplicateInputs(
 	const float BrakeInput, const float HandbrakeInput, const float RollInput,
 	const float PitchInput, const float YawInput, const int32 CurrentGear, DatasType& ReplicatedDatas)
 {
+	using namespace Chaos;
+
 	ReplicatedDatas.VehicleInputs.SteeringInput = SteeringInput;
 	ReplicatedDatas.VehicleInputs.ThrottleInput = ThrottleInput;
 	ReplicatedDatas.VehicleInputs.BrakeInput = BrakeInput;
@@ -379,6 +384,8 @@ void FVehicleInputsHistory::ReplicateInputs(
 
 void FVehicleInputsHistory::SmoothInputs(const int32 ApplyFrame, const int32 ValidFrame, const bool bApplySmoothing)
 {
+	using namespace Chaos;
+
 	if (!bIsLocalHistory && VehicleSimulation)
 	{
 		const float HalfFrames = (float)NumFrames / 2.0f;
@@ -420,6 +427,8 @@ void FVehicleInputsHistory::SmoothInputs(const int32 ApplyFrame, const int32 Val
 
 bool FVehicleInputsHistory::ApplyInputs(const int32 ApplyFrame, const bool bResetSolver)
 {
+	using namespace Chaos;
+
 	if (VehicleSimulation )
 	{
 		if(Super::ApplyInputs(ApplyFrame, bResetSolver))
