Make some vehicle properties public

diff --git a/Engine/Plugins/Experimental/ChaosVehiclesPlugin/Source/ChaosVehicles/Public/ChaosVehicleMovementComponent.h b/Engine/Plugins/Experimental/ChaosVehiclesPlugin/Source/ChaosVehicles/Public/ChaosVehicleMovementComponent.h
--- a/Engine/Plugins/Experimental/ChaosVehiclesPlugin/Source/ChaosVehicles/Public/ChaosVehicleMovementComponent.h
+++ b/Engine/Plugins/Experimental/ChaosVehiclesPlugin/Source/ChaosVehicles/Public/ChaosVehicleMovementComponent.h
@@ -1386,6 +1386,7 @@ protected:
 	UPROPERTY(EditAnywhere, Category = VehicleInput)
 	float WrongDirectionThreshold;
 
+public:
 	// Rate at which input throttle can rise and fall
 	UPROPERTY(EditAnywhere, Category = VehicleInput, AdvancedDisplay)
 	FVehicleInputRateConfig ThrottleInputRate;
@@ -1414,6 +1415,7 @@ protected:
 	UPROPERTY(EditAnywhere, Category = VehicleInput, AdvancedDisplay)
 	FVehicleInputRateConfig YawInputRate;
 
+protected:
 
 	// input related
 
