Removing redundant code

diff --git a/Engine/Plugins/Experimental/ChaosVehiclesPlugin/Source/ChaosVehicles/Private/ChaosVehicleMovementComponent.cpp b/Engine/Plugins/Experimental/ChaosVehiclesPlugin/Source/ChaosVehicles/Private/ChaosVehicleMovementComponent.cpp
--- a/Engine/Plugins/Experimental/ChaosVehiclesPlugin/Source/ChaosVehicles/Private/ChaosVehicleMovementComponent.cpp
+++ b/Engine/Plugins/Experimental/ChaosVehiclesPlugin/Source/ChaosVehicles/Private/ChaosVehicleMovementComponent.cpp
@@ -2245,15 +2245,6 @@ void UChaosVehicleMovementComponent::ResetVehicleState()
 
 void UChaosVehicleMovementComponent::FinalizeSimCallbackData(FChaosVehicleManagerAsyncInput& Input)
 {
-	bool bIsPhysicsStateCreated = true;
-	if (GetSkeletalMesh())
-	{
-		if (USkeletalMeshComponent* SkelMesh = GetSkeletalMesh())
-		{
-			bIsPhysicsStateCreated = SkelMesh->IsPhysicsStateCreated();
-		}
-	}
-
 	CurAsyncInput = nullptr;
 	CurAsyncOutput = nullptr;
 }
