Vehicle wheel controller crash workaround with logging

diff --git a/Engine/Plugins/Experimental/ChaosVehiclesPlugin/Source/ChaosVehicles/Private/AnimNode_StageCoachWheelController.cpp b/Engine/Plugins/Experimental/ChaosVehiclesPlugin/Source/ChaosVehicles/Private/AnimNode_StageCoachWheelController.cpp
--- a/Engine/Plugins/Experimental/ChaosVehiclesPlugin/Source/ChaosVehicles/Private/AnimNode_StageCoachWheelController.cpp
+++ b/Engine/Plugins/Experimental/ChaosVehiclesPlugin/Source/ChaosVehicles/Private/AnimNode_StageCoachWheelController.cpp
@@ -55,35 +55,42 @@ void FAnimNode_StageCoachWheelController::EvaluateSkeletalControl_AnyThread(FCom
 	for(const FWheelLookupData& Wheel : Wheels)
 	{
 		if (Wheel.BoneReference.IsValidToEvaluate(BoneContainer))
-		{
-			FCompactPoseBoneIndex WheelSimBoneIndex = Wheel.BoneReference.GetCompactPoseIndex(BoneContainer);
-
-			// the way we apply transform is same as FMatrix or FTransform
-			// we apply scale first, and rotation, and translation
-			// if you'd like to translate first, you'll need two nodes that first node does translate and second nodes to rotate. 
-			FTransform NewBoneTM = Output.Pose.GetComponentSpaceTransform(WheelSimBoneIndex);
-
-			FAnimationRuntime::ConvertCSTransformToBoneSpace(Output.AnimInstanceProxy->GetComponentTransform(), Output.Pose, NewBoneTM, WheelSimBoneIndex, BCS_ComponentSpace);
-
-			// Apply rotation offset
-			const FQuat BoneQuat(WheelAnimData[Wheel.WheelIndex].RotOffset);
-			NewBoneTM.SetRotation(BoneQuat * NewBoneTM.GetRotation());
-
-			// Apply loc offset
-			NewBoneTM.AddToTranslation(WheelAnimData[Wheel.WheelIndex].LocOffset);
-
-			// Convert back to Component Space.
-			FAnimationRuntime::ConvertBoneSpaceTransformToCS(Output.AnimInstanceProxy->GetComponentTransform(), Output.Pose, NewBoneTM, WheelSimBoneIndex, BCS_ComponentSpace);
-
-			// add back to it
-			OutBoneTransforms.Add(FBoneTransform(WheelSimBoneIndex, NewBoneTM));
+		{ 
+			if (Wheel.WheelIndex < WheelAnimData.Num())
+			{
+				FCompactPoseBoneIndex WheelSimBoneIndex = Wheel.BoneReference.GetCompactPoseIndex(BoneContainer);
+
+				// the way we apply transform is same as FMatrix or FTransform
+				// we apply scale first, and rotation, and translation
+				// if you'd like to translate first, you'll need two nodes that first node does translate and second nodes to rotate. 
+				FTransform NewBoneTM = Output.Pose.GetComponentSpaceTransform(WheelSimBoneIndex);
+
+				FAnimationRuntime::ConvertCSTransformToBoneSpace(Output.AnimInstanceProxy->GetComponentTransform(), Output.Pose, NewBoneTM, WheelSimBoneIndex, BCS_ComponentSpace);
+
+				// Apply rotation offset
+				const FQuat BoneQuat(WheelAnimData[Wheel.WheelIndex].RotOffset);
+				NewBoneTM.SetRotation(BoneQuat * NewBoneTM.GetRotation());
+
+				// Apply loc offset
+				NewBoneTM.AddToTranslation(WheelAnimData[Wheel.WheelIndex].LocOffset);
+
+				// Convert back to Component Space.
+				FAnimationRuntime::ConvertBoneSpaceTransformToCS(Output.AnimInstanceProxy->GetComponentTransform(), Output.Pose, NewBoneTM, WheelSimBoneIndex, BCS_ComponentSpace);
+
+				// add back to it
+				OutBoneTransforms.Add(FBoneTransform(WheelSimBoneIndex, NewBoneTM));
+			}
+			else
+			{
+				UE_LOG(LogChaos, Error, TEXT("Invalid condition Wheel.WheelIndex (%d) >= WheelAnimData.Num (%d)"), Wheel.WheelIndex, WheelAnimData.Num());
+			}
 		}
 	}
 
 #if ANIM_TRACE_ENABLED
 	for (const FWheelLookupData& Wheel : Wheels)
 	{
-		if (Wheel.BoneReference.BoneIndex != INDEX_NONE)
+		if ((Wheel.BoneReference.BoneIndex != INDEX_NONE) && (Wheel.WheelIndex < WheelAnimData.Num()))
 		{
 			TRACE_ANIM_NODE_VALUE(Output, *FString::Printf(TEXT("Wheel %d Name"), Wheel.WheelIndex), *Wheel.BoneReference.BoneName.ToString());
 			TRACE_ANIM_NODE_VALUE(Output, *FString::Printf(TEXT("Wheel %d Rotation Offset"), Wheel.WheelIndex), WheelAnimData[Wheel.WheelIndex].RotOffset);
diff --git a/Engine/Plugins/Experimental/ChaosVehiclesPlugin/Source/ChaosVehicles/Private/AnimNode_WheelController.cpp b/Engine/Plugins/Experimental/ChaosVehiclesPlugin/Source/ChaosVehicles/Private/AnimNode_WheelController.cpp
--- a/Engine/Plugins/Experimental/ChaosVehiclesPlugin/Source/ChaosVehicles/Private/AnimNode_WheelController.cpp
+++ b/Engine/Plugins/Experimental/ChaosVehiclesPlugin/Source/ChaosVehicles/Private/AnimNode_WheelController.cpp
@@ -52,34 +52,41 @@ void FAnimNode_WheelController::EvaluateSkeletalControl_AnyThread(FComponentSpac
 	{
 		if (Wheel.BoneReference.IsValidToEvaluate(BoneContainer))
 		{
-			FCompactPoseBoneIndex WheelSimBoneIndex = Wheel.BoneReference.GetCompactPoseIndex(BoneContainer);
-
-			// the way we apply transform is same as FMatrix or FTransform
-			// we apply scale first, and rotation, and translation
-			// if you'd like to translate first, you'll need two nodes that first node does translate and second nodes to rotate. 
-			FTransform NewBoneTM = Output.Pose.GetComponentSpaceTransform(WheelSimBoneIndex);
-
-			FAnimationRuntime::ConvertCSTransformToBoneSpace(Output.AnimInstanceProxy->GetComponentTransform(), Output.Pose, NewBoneTM, WheelSimBoneIndex, BCS_ComponentSpace);
-
-			// Apply rotation offset
-			const FQuat BoneQuat(WheelAnimData[Wheel.WheelIndex].RotOffset);
-			NewBoneTM.SetRotation(BoneQuat * NewBoneTM.GetRotation());
-
-			// Apply loc offset
-			NewBoneTM.AddToTranslation(WheelAnimData[Wheel.WheelIndex].LocOffset);
-
-			// Convert back to Component Space.
-			FAnimationRuntime::ConvertBoneSpaceTransformToCS(Output.AnimInstanceProxy->GetComponentTransform(), Output.Pose, NewBoneTM, WheelSimBoneIndex, BCS_ComponentSpace);
-
-			// add back to it
-			OutBoneTransforms.Add(FBoneTransform(WheelSimBoneIndex, NewBoneTM));
+			if (Wheel.WheelIndex < WheelAnimData.Num())
+			{
+				FCompactPoseBoneIndex WheelSimBoneIndex = Wheel.BoneReference.GetCompactPoseIndex(BoneContainer);
+
+				// the way we apply transform is same as FMatrix or FTransform
+				// we apply scale first, and rotation, and translation
+				// if you'd like to translate first, you'll need two nodes that first node does translate and second nodes to rotate. 
+				FTransform NewBoneTM = Output.Pose.GetComponentSpaceTransform(WheelSimBoneIndex);
+
+				FAnimationRuntime::ConvertCSTransformToBoneSpace(Output.AnimInstanceProxy->GetComponentTransform(), Output.Pose, NewBoneTM, WheelSimBoneIndex, BCS_ComponentSpace);
+
+				// Apply rotation offset
+				const FQuat BoneQuat(WheelAnimData[Wheel.WheelIndex].RotOffset);
+				NewBoneTM.SetRotation(BoneQuat * NewBoneTM.GetRotation());
+
+				// Apply loc offset
+				NewBoneTM.AddToTranslation(WheelAnimData[Wheel.WheelIndex].LocOffset);
+
+				// Convert back to Component Space.
+				FAnimationRuntime::ConvertBoneSpaceTransformToCS(Output.AnimInstanceProxy->GetComponentTransform(), Output.Pose, NewBoneTM, WheelSimBoneIndex, BCS_ComponentSpace);
+
+				// add back to it
+				OutBoneTransforms.Add(FBoneTransform(WheelSimBoneIndex, NewBoneTM));
+			}
+			else
+			{
+				UE_LOG(LogChaos, Error, TEXT("Invalid condition Wheel.WheelIndex (%d) >= WheelAnimData.Num (%d)"), Wheel.WheelIndex, WheelAnimData.Num());
+			}
 		}
 	}
 
 #if ANIM_TRACE_ENABLED
 	for (const FWheelLookupData& Wheel : Wheels)
 	{
-		if (Wheel.BoneReference.BoneIndex != INDEX_NONE)
+		if ((Wheel.BoneReference.BoneIndex != INDEX_NONE) && (Wheel.WheelIndex < WheelAnimData.Num()))
 		{
 			TRACE_ANIM_NODE_VALUE(Output, *FString::Printf(TEXT("Wheel %d Name"), Wheel.WheelIndex), *Wheel.BoneReference.BoneName.ToString());
 			TRACE_ANIM_NODE_VALUE(Output, *FString::Printf(TEXT("Wheel %d Rotation Offset"), Wheel.WheelIndex), WheelAnimData[Wheel.WheelIndex].RotOffset);
