Network physics component

diff --git a/Engine/Plugins/Experimental/ChaosVehiclesPlugin/Source/ChaosVehicles/Private/ChaosVehicleManagerAsyncCallback.cpp b/Engine/Plugins/Experimental/ChaosVehiclesPlugin/Source/ChaosVehicles/Private/ChaosVehicleManagerAsyncCallback.cpp
--- a/Engine/Plugins/Experimental/ChaosVehiclesPlugin/Source/ChaosVehicles/Private/ChaosVehicleManagerAsyncCallback.cpp
+++ b/Engine/Plugins/Experimental/ChaosVehiclesPlugin/Source/ChaosVehicles/Private/ChaosVehicleManagerAsyncCallback.cpp
@@ -1,9 +1,9 @@
 // Copyright Epic Games, Inc. All Rights Reserved.
 
 #include "ChaosVehicleManagerAsyncCallback.h"
-
 #include "ChaosVehicleMovementComponent.h"
 #include "PBDRigidsSolver.h"
+#include "TransmissionSystem.h"
 #include "Chaos/ParticleHandleFwd.h"
 #include "PhysicsProxy/SingleParticlePhysicsProxy.h"
 
@@ -20,6 +20,44 @@ FName FChaosVehicleManagerAsyncCallback::GetFNameForStatId() const
 /**
  * Callback from Physics thread
  */
+
+void FChaosVehicleManagerAsyncCallback::ProcessInputs_Internal(int32 PhysicsStep) 
+{
+	const FChaosVehicleManagerAsyncInput* AsyncInput = GetConsumerInput_Internal();
+	if (AsyncInput == nullptr)
+	{
+		return;
+	}
+
+	for (const TUniquePtr<FChaosVehicleAsyncInput>& VehicleInput : AsyncInput->VehicleInputs)
+	{
+		UChaosVehicleSimulation* VehicleSim = VehicleInput->Vehicle->VehicleSimulationPT.Get();
+
+		if (VehicleSim == nullptr || !VehicleInput->Vehicle->bUsingNetworkPhysicsPrediction)
+		{
+			continue;
+		}
+		bool bIsResimming = false;
+		if (FPhysScene* PhysScene = VehicleInput->Vehicle->GetWorld()->GetPhysicsScene())
+		{
+			if (Chaos::FPhysicsSolver* LocalSolver = PhysScene->GetSolver())
+			{
+				bIsResimming = LocalSolver->GetEvolution()->IsResimming();
+			}
+		}
+
+		APlayerController* PlayerController = VehicleInput->Vehicle->GetPlayerController();
+		if(PlayerController && PlayerController->IsLocalController() && !bIsResimming)
+		{ 
+			VehicleSim->VehicleInputs = VehicleInput->PhysicsInputs.NetworkInputs.VehicleInputs;
+		}
+		else
+		{
+			VehicleInput->PhysicsInputs.NetworkInputs.VehicleInputs = VehicleSim->VehicleInputs;
+		}
+	}
+}
+
 void FChaosVehicleManagerAsyncCallback::OnPreSimulate_Internal()
 {
 	using namespace Chaos;
@@ -75,7 +113,6 @@ void FChaosVehicleManagerAsyncCallback::OnPreSimulate_Internal()
 
 		bool bWake = false;
 		OutputVehiclesBatch[Idx] = VehicleInput.Simulate(World, DeltaTime, SimTime, bWake);
-
 	};
 
 	bool ForceSingleThread = !GVehicleDebugParams.EnableMultithreading;
@@ -94,11 +131,338 @@ void FChaosVehicleManagerAsyncCallback::OnPreSimulate_Internal()
 	}
 }
 
-/**
- * Contact modification currently unused
- */
-void FChaosVehicleManagerAsyncCallback::OnContactModification_Internal(Chaos::FCollisionContactModifier& Modifications)
+TUniquePtr<FChaosVehicleAsyncOutput> FChaosVehicleAsyncInput::Simulate(UWorld* World, const float DeltaSeconds, const float TotalSeconds, bool& bWakeOut) const
+{
+	TUniquePtr<FChaosVehicleAsyncOutput> Output = MakeUnique<FChaosVehicleAsyncOutput>();
+
+	//UE_LOG(LogChaos, Warning, TEXT("Vehicle Physics Thread Tick %f"), DeltaSeconds);
+
+	//support nullptr because it allows us to go wide on filling the async inputs
+	if (Proxy == nullptr)
+	{
+		return Output;
+	}
+
+	// We now have access to the physics representation of the chassis on the physics thread async tick
+	Chaos::FRigidBodyHandle_Internal* Handle = Proxy->GetPhysicsThreadAPI();
+
+	// FILL OUTPUT DATA HERE THAT WILL GET PASSED BACK TO THE GAME THREAD
+	Vehicle->VehicleSimulationPT->TickVehicle(World, DeltaSeconds, *this, *Output.Get(), Handle);
+
+	Output->bValid = true;
+
+	return MoveTemp(Output);
+}
+
+void FChaosVehicleAsyncInput::ApplyDeferredForces(Chaos::FRigidBodyHandle_Internal* RigidHandle) const
+{
+	check(Vehicle);
+	check(Vehicle->VehicleSimulationPT);
+	Vehicle->VehicleSimulationPT->ApplyDeferredForces(RigidHandle);
+}
+
+bool FNetworkVehicleInputs::NetSerialize(FArchive& Ar, class UPackageMap* Map, bool& bOutSuccess)
+{
+	FNetworkPhysicsDatas::SerializeFrames(Ar);
+
+	Ar << VehicleInputs.SteeringInput;
+	Ar << VehicleInputs.ThrottleInput;
+	Ar << VehicleInputs.BrakeInput;
+	Ar << VehicleInputs.PitchInput;
+	Ar << VehicleInputs.RollInput;
+	Ar << VehicleInputs.YawInput;
+	Ar << VehicleInputs.HandbrakeInput;
+
+	Ar << TransmissionChangeTime;
+	Ar << TransmissionCurrentGear;
+	Ar << TransmissionTargetGear;
+
+	bOutSuccess = true;
+	return bOutSuccess;
+}
+
+void FNetworkVehicleInputs::ApplyDatas(UActorComponent* NetworkComponent) const
+{
+	if (UChaosVehicleSimulation* VehicleSimulation = Cast<UChaosVehicleMovementComponent>(NetworkComponent)->VehicleSimulationPT.Get())
+	{
+		VehicleSimulation->VehicleInputs = VehicleInputs;
+
+		if (TUniquePtr<Chaos::FSimpleWheeledVehicle>& Vehicle = VehicleSimulation->PVehicle)
+		{
+			if (Vehicle->HasTransmission())
+			{
+				Chaos::FSimpleTransmissionSim& Transmission = VehicleSimulation->PVehicle->GetTransmission();
+				Transmission.SetCurrentGear(TransmissionCurrentGear);
+				Transmission.SetTargetGear(TransmissionTargetGear);
+				Transmission.SetCurrentGearChangeTime(TransmissionChangeTime);
+			}
+		}
+#if DEBUG_NETWORK_PHYSICS
+		int32 SolverFrame = INDEX_NONE;
+		if (NetworkComponent->GetWorld() && NetworkComponent->GetWorld()->GetPhysicsScene())
+		{
+			if (Chaos::FPhysicsSolver* PhysicsSolver = NetworkComponent->GetWorld()->GetPhysicsScene()->GetSolver())
+			{
+				SolverFrame = PhysicsSolver->GetCurrentFrame();
+			}
+		}
+
+		if (NetworkComponent->GetWorld()->IsNetMode(NM_ListenServer))
+		{
+			UE_LOG(LogTemp, Log, TEXT("SERVER | PT | ApplyDatas | Report replicated inputs at frame %d %d: Throttle = %f Brake = %f Roll = %f Pitch = %f Yaw = %f Steering = %f Handbrake = %f Gear = %d | VehicleInputs size = %d | ControlInputs size = %d | NetworkInputs = %d"),
+				LocalFrame, SolverFrame, VehicleInputs.ThrottleInput, VehicleInputs.BrakeInput, VehicleInputs.RollInput, VehicleInputs.PitchInput,
+				VehicleInputs.YawInput, VehicleInputs.SteeringInput, VehicleInputs.HandbrakeInput, TransmissionTargetGear, sizeof(FVehicleInputs) * 8, sizeof(FControlInputs) * 8, sizeof(FNetworkVehicleInputs) * 8);
+		}
+		else
+		{
+			UE_LOG(LogTemp, Log, TEXT("CLIENT | PT | ApplyDatas | Report replicated inputs at frame %d %d: Throttle = %f Brake = %f Roll = %f Pitch = %f Yaw = %f Steering = %f Handbrake = %f Gear = %d"),
+				LocalFrame, SolverFrame, VehicleInputs.ThrottleInput, VehicleInputs.BrakeInput, VehicleInputs.RollInput, VehicleInputs.PitchInput,
+				VehicleInputs.YawInput, VehicleInputs.SteeringInput, VehicleInputs.HandbrakeInput, TransmissionTargetGear);
+		}
+#endif
+	}
+}
+
+void FNetworkVehicleInputs::BuildDatas(const UActorComponent* NetworkComponent)
+{
+	if (NetworkComponent)
+	{
+		if (const UChaosVehicleSimulation* VehicleSimulation = Cast<const UChaosVehicleMovementComponent>(NetworkComponent)->VehicleSimulationPT.Get())
+		{
+			VehicleInputs = VehicleSimulation->VehicleInputs;
+
+			if (const TUniquePtr<Chaos::FSimpleWheeledVehicle>& Vehicle = VehicleSimulation->PVehicle)
+			{
+				if (Vehicle->HasTransmission())
+				{
+					Chaos::FSimpleTransmissionSim& Transmission = Vehicle->GetTransmission();
+					TransmissionCurrentGear = Transmission.GetCurrentGear();
+					TransmissionTargetGear = Transmission.GetTargetGear();
+					TransmissionChangeTime = Transmission.GetCurrentGearChangeTime();
+				}
+			}
+#if DEBUG_NETWORK_PHYSICS
+			if(NetworkComponent->GetWorld()->IsNetMode(NM_ListenServer))
+			{
+				UE_LOG(LogTemp, Log, TEXT("SERVER | PT | BuildDatas | Extract local inputs at frame %d : Throttle = %f Brake = %f Roll = %f Pitch = %f Yaw = %f Steering = %f Handbrake = %f Gear = %d"),
+					LocalFrame, VehicleInputs.ThrottleInput, VehicleInputs.BrakeInput, VehicleInputs.RollInput, VehicleInputs.PitchInput,
+					VehicleInputs.YawInput, VehicleInputs.SteeringInput, VehicleInputs.HandbrakeInput, TransmissionTargetGear);
+			}
+			else
+			{
+				UE_LOG(LogTemp, Log, TEXT("CLIENT | PT | BuildDatas | Extract local inputs at frame %d : Throttle = %f Brake = %f Roll = %f Pitch = %f Yaw = %f Steering = %f Handbrake = %f Gear = %d"),
+					LocalFrame, VehicleInputs.ThrottleInput, VehicleInputs.BrakeInput, VehicleInputs.RollInput, VehicleInputs.PitchInput,
+					VehicleInputs.YawInput, VehicleInputs.SteeringInput, VehicleInputs.HandbrakeInput, TransmissionTargetGear);
+			}
+#endif
+		}
+	}
+}
+
+void FNetworkVehicleInputs::InterpolateDatas(const FNetworkVehicleInputs& MinDatas, const FNetworkVehicleInputs& MaxDatas)
+{
+	const float LerpFactor = (LocalFrame - MinDatas.LocalFrame) / (MaxDatas.LocalFrame - MinDatas.LocalFrame);
+
+	TransmissionChangeTime = FMath::Lerp(MinDatas.TransmissionChangeTime, MaxDatas.TransmissionChangeTime, LerpFactor);
+	TransmissionCurrentGear = LerpFactor < 0.5 ? MinDatas.TransmissionCurrentGear : MaxDatas.TransmissionCurrentGear;
+	TransmissionTargetGear = LerpFactor < 0.5 ? MinDatas.TransmissionTargetGear : MaxDatas.TransmissionTargetGear;
+
+	VehicleInputs.BrakeInput = FMath::Lerp(MinDatas.VehicleInputs.BrakeInput, MaxDatas.VehicleInputs.BrakeInput, LerpFactor);
+	VehicleInputs.HandbrakeInput = FMath::Lerp(MinDatas.VehicleInputs.HandbrakeInput, MaxDatas.VehicleInputs.HandbrakeInput, LerpFactor);
+	VehicleInputs.PitchInput = FMath::Lerp(MinDatas.VehicleInputs.PitchInput, MaxDatas.VehicleInputs.PitchInput, LerpFactor);
+	VehicleInputs.RollInput = FMath::Lerp(MinDatas.VehicleInputs.RollInput, MaxDatas.VehicleInputs.RollInput, LerpFactor);
+	VehicleInputs.ThrottleInput = FMath::Lerp(MinDatas.VehicleInputs.ThrottleInput, MaxDatas.VehicleInputs.ThrottleInput, LerpFactor);
+	VehicleInputs.SteeringInput = FMath::Lerp(MinDatas.VehicleInputs.SteeringInput, MaxDatas.VehicleInputs.SteeringInput, LerpFactor);
+	VehicleInputs.YawInput = FMath::Lerp(MinDatas.VehicleInputs.YawInput, MaxDatas.VehicleInputs.YawInput, LerpFactor);
+
+	VehicleInputs.ParkingEnabled = LerpFactor < 0.5 ? MinDatas.VehicleInputs.ParkingEnabled : MaxDatas.VehicleInputs.ParkingEnabled;
+	VehicleInputs.GearDownInput = LerpFactor < 0.5 ? MinDatas.VehicleInputs.GearDownInput : MaxDatas.VehicleInputs.GearDownInput;
+	VehicleInputs.GearUpInput = LerpFactor < 0.5 ? MinDatas.VehicleInputs.GearUpInput : MaxDatas.VehicleInputs.GearUpInput;
+	VehicleInputs.TransmissionType = LerpFactor < 0.5 ? MinDatas.VehicleInputs.TransmissionType : MaxDatas.VehicleInputs.TransmissionType;
+}
+
+bool FNetworkVehicleStates::NetSerialize(FArchive& Ar, class UPackageMap* Map, bool& bOutSuccess)
+{
+	FNetworkPhysicsDatas::SerializeFrames(Ar);
+
+	Ar << StateLastVelocity;
+	Ar << EngineOmega;
+
+	int32 NumWheels = WheelsOmega.Num();
+	Ar << NumWheels;
+
+	int32 NumLength = SuspensionAveragedLength.Num();
+	Ar << NumLength;
+
+	if (Ar.IsLoading())
+	{
+		WheelsOmega.SetNum(NumWheels);
+		WheelsAngularPosition.SetNum(NumWheels);
+		SuspensionLastDisplacement.SetNum(NumWheels);
+		SuspensionLastSpringLength.SetNum(NumWheels);
+		SuspensionAveragedCount.SetNum(NumWheels);
+		SuspensionAveragedNum.SetNum(NumWheels);
+		SuspensionAveragedLength.SetNum(NumLength);
+	}
+
+	for (int32 WheelIdx = 0; WheelIdx < NumWheels; ++WheelIdx)
+	{
+		Ar << WheelsOmega[WheelIdx];
+		Ar << WheelsAngularPosition[WheelIdx];
+		Ar << SuspensionLastDisplacement[WheelIdx];
+		Ar << SuspensionLastSpringLength[WheelIdx];
+		Ar << SuspensionAveragedCount[WheelIdx];
+		Ar << SuspensionAveragedNum[WheelIdx];
+	}
+
+	for (int32 LengthIdx = 0; LengthIdx < NumLength; ++LengthIdx)
+	{
+		Ar << SuspensionAveragedLength[LengthIdx];
+	}
+	return true;
+}
+
+void FNetworkVehicleStates::ApplyDatas(UActorComponent* NetworkComponent) const
 {
+	if (UChaosVehicleSimulation* VehicleSimulation = Cast<UChaosVehicleMovementComponent>(NetworkComponent)->VehicleSimulationPT.Get())
+	{
+		VehicleSimulation->VehicleState.LastFrameVehicleLocalVelocity = StateLastVelocity;
+		
+		if (TUniquePtr<Chaos::FSimpleWheeledVehicle>& Vehicle = VehicleSimulation->PVehicle)
+		{
+			Vehicle->GetEngine().SetEngineOmega(EngineOmega);
 
+#if DEBUG_NETWORK_PHYSICS
+			if (NetworkComponent->GetWorld()->IsNetMode(NM_ListenServer))
+			{
+				UE_LOG(LogTemp, Log, TEXT("SERVER | PT | ApplyDatas | Report replicated states at frame %d %d : Omega = %f"),
+					LocalFrame, ServerFrame, EngineOmega);
+			}
+			else
+			{
+				UE_LOG(LogTemp, Log, TEXT("CLIENT | PT | ApplyDatas| Report replicated states at frame %d %d : Omega = %f"),
+					LocalFrame, ServerFrame, EngineOmega);
+			}
+#endif
+		
+			int32 LengthCount = 0;
+			for (int32 WheelIdx = 0, NumWheels = Vehicle->Wheels.Num(); WheelIdx < NumWheels; ++WheelIdx)
+			{
+				Chaos::FSimpleSuspensionSim& Suspension = Vehicle->GetSuspension(WheelIdx);
+				Suspension.SetLastSpringLength(SuspensionLastSpringLength[WheelIdx]);
+				Suspension.SetLastDisplacement(SuspensionLastDisplacement[WheelIdx]);
+				Suspension.SetAveragingCount(SuspensionAveragedCount[WheelIdx]);
+				Suspension.SetAveragingNum(SuspensionAveragedNum[WheelIdx]);
+		
+				for (int32 LengthIdx = 0; LengthIdx < Suspension.GetAveragingNum(); ++LengthIdx)
+				{
+					Suspension.SetAveragingLength(LengthIdx, SuspensionAveragedLength[LengthCount++]);
+				}
+		
+				Chaos::FSimpleWheelSim& Wheel = Vehicle->GetWheel(WheelIdx);
+				Wheel.Omega = WheelsOmega[WheelIdx];
+				Wheel.AngularPosition = WheelsAngularPosition[WheelIdx];
+			}
+		}
+	}
 }
 
+void FNetworkVehicleStates::BuildDatas(const UActorComponent* NetworkComponent)
+{
+	if (NetworkComponent)
+	{
+		if (const UChaosVehicleSimulation* VehicleSimulation = Cast<const UChaosVehicleMovementComponent>(NetworkComponent)->VehicleSimulationPT.Get())
+		{
+			StateLastVelocity = VehicleSimulation->VehicleState.LastFrameVehicleLocalVelocity;
+			if (const TUniquePtr<Chaos::FSimpleWheeledVehicle>& Vehicle = VehicleSimulation->PVehicle)
+			{
+				EngineOmega = Vehicle->GetEngine().GetEngineOmega();
+#if DEBUG_NETWORK_PHYSICS
+				if (NetworkComponent->GetWorld()->IsNetMode(NM_ListenServer))
+				{
+					UE_LOG(LogTemp, Log, TEXT("SERVER | PT | BuildDatas | Extract local states at frame %d %d : Omega = %f"),
+						LocalFrame, ServerFrame, EngineOmega);
+				}
+				else
+				{
+					UE_LOG(LogTemp, Log, TEXT("CLIENT | PT | BuildDatas| Extract local states at frame %d %d : Omega = %f"),
+						LocalFrame, ServerFrame, EngineOmega);
+				}
+#endif
+				
+				const int32 NumWheels = Vehicle->Wheels.Num();
+				
+				int32 NumLength = 0;
+				for (int32 WheelIdx = 0; WheelIdx < NumWheels; ++WheelIdx)
+				{
+					NumLength += Vehicle->GetSuspension(WheelIdx).GetAveragingNum();
+				}
+				
+				SuspensionLastSpringLength.SetNum(NumWheels);
+				SuspensionLastDisplacement.SetNum(NumWheels);
+				SuspensionAveragedCount.SetNum(NumWheels);
+				SuspensionAveragedNum.SetNum(NumWheels);
+				SuspensionAveragedLength.SetNum(NumLength);
+				WheelsAngularPosition.SetNum(NumWheels);
+				WheelsOmega.SetNum(NumWheels);
+				
+				int32 LengthCount = 0;
+				for (int32 WheelIdx = 0; WheelIdx < Vehicle->Wheels.Num(); ++WheelIdx)
+				{
+					Chaos::FSimpleSuspensionSim& Suspension = Vehicle->GetSuspension(WheelIdx);
+					SuspensionLastSpringLength[WheelIdx] = Suspension.GetLastSpringLength();
+					SuspensionLastDisplacement[WheelIdx] = Suspension.GetLastDisplacement();
+					SuspensionAveragedCount[WheelIdx] = Suspension.GetAveragingCount();
+					SuspensionAveragedNum[WheelIdx] = Suspension.GetAveragingNum();
+				
+					for (int32 LengthIdx = 0; LengthIdx < Suspension.GetAveragingNum(); ++LengthIdx)
+					{
+						SuspensionAveragedLength[LengthCount++] = Suspension.GetAveragingLength(LengthIdx);
+					}
+				
+					Chaos::FSimpleWheelSim& Wheel = Vehicle->GetWheel(WheelIdx);
+					WheelsOmega[WheelIdx] = Wheel.Omega;
+					WheelsAngularPosition[WheelIdx] = Wheel.AngularPosition;
+				}
+			}
+		}
+	}
+}
+void FNetworkVehicleStates::InterpolateDatas(const FNetworkVehicleStates& MinDatas, const FNetworkVehicleStates& MaxDatas)
+{
+	const float LerpFactor = (LocalFrame - MinDatas.LocalFrame) / (MaxDatas.LocalFrame - MinDatas.LocalFrame);
+
+	StateLastVelocity = FMath::Lerp(MinDatas.StateLastVelocity, MaxDatas.StateLastVelocity, LerpFactor);
+	EngineOmega = FMath::Lerp(MinDatas.EngineOmega, MaxDatas.EngineOmega, LerpFactor);
+
+	int32 NumWheels = FMath::Min(MinDatas.WheelsOmega.Num(), MaxDatas.WheelsOmega.Num());
+
+	WheelsOmega.SetNum(NumWheels, false);
+	WheelsAngularPosition.SetNum(NumWheels, false);
+
+	SuspensionLastDisplacement.SetNum(NumWheels, false);
+	SuspensionLastSpringLength.SetNum(NumWheels, false);
+	SuspensionAveragedCount.SetNum(NumWheels, false);
+	SuspensionAveragedNum.SetNum(NumWheels, false);
+
+	int32 NumLength = 0;
+	for (int32 WheelIdx = 0; WheelIdx < NumWheels; ++WheelIdx)
+	{
+		WheelsOmega[WheelIdx] = FMath::Lerp(MinDatas.WheelsOmega[WheelIdx], MaxDatas.WheelsOmega[WheelIdx], LerpFactor);
+		WheelsAngularPosition[WheelIdx] = FMath::Lerp(MinDatas.WheelsAngularPosition[WheelIdx], MaxDatas.WheelsAngularPosition[WheelIdx], LerpFactor);
+
+		SuspensionLastDisplacement[WheelIdx] = FMath::Lerp(MinDatas.SuspensionLastDisplacement[WheelIdx], MaxDatas.SuspensionLastDisplacement[WheelIdx], LerpFactor);
+		SuspensionLastSpringLength[WheelIdx] = FMath::Lerp(MinDatas.SuspensionLastSpringLength[WheelIdx], MaxDatas.SuspensionLastSpringLength[WheelIdx], LerpFactor);
+		SuspensionAveragedCount[WheelIdx] = FMath::Min( MinDatas.SuspensionAveragedCount[WheelIdx], MaxDatas.SuspensionAveragedCount[WheelIdx]);
+		SuspensionAveragedNum[WheelIdx] = FMath::Min(MinDatas.SuspensionAveragedNum[WheelIdx], MaxDatas.SuspensionAveragedNum[WheelIdx]);
+
+		NumLength += SuspensionAveragedNum[WheelIdx];
+	}
+	SuspensionAveragedLength.SetNum(NumLength, false);
+	for (int32 LengthIdx = 0; LengthIdx < NumLength; ++LengthIdx)
+	{
+		SuspensionAveragedLength[LengthIdx] = FMath::Lerp(MinDatas.SuspensionAveragedLength[LengthIdx], MaxDatas.SuspensionAveragedLength[LengthIdx], LerpFactor);
+	}
+}
+
+
diff --git a/Engine/Plugins/Experimental/ChaosVehiclesPlugin/Source/ChaosVehicles/Private/ChaosVehicleMovementComponent.cpp b/Engine/Plugins/Experimental/ChaosVehiclesPlugin/Source/ChaosVehicles/Private/ChaosVehicleMovementComponent.cpp
--- a/Engine/Plugins/Experimental/ChaosVehiclesPlugin/Source/ChaosVehicles/Private/ChaosVehicleMovementComponent.cpp
+++ b/Engine/Plugins/Experimental/ChaosVehiclesPlugin/Source/ChaosVehicles/Private/ChaosVehicleMovementComponent.cpp
@@ -73,19 +73,6 @@ FAutoConsoleVariableRef CVarChaosVehiclesEnableMultithreading(TEXT("p.Vehicle.En
 FAutoConsoleVariableRef CVarChaosVehiclesControlInputWakeTolerance(TEXT("p.Vehicle.ControlInputWakeTolerance"), GVehicleDebugParams.ControlInputWakeTolerance, TEXT("Set the control input wake tolerance."));
 
 
-FVehicleInputsDatas::FVehicleInputsDatas()
-	: TransmissionChangeTime(0.0f)
-	, TransmissionCurrentGear(0)
-	, TransmissionTargetGear(0)
-{
-}
-
-FVehicleStatesDatas::FVehicleStatesDatas()
-	: StateLastVelocity(FVector::ZeroVector)
-	, EngineOmega(0.0f)
-{
-}
-
 void FVehicleState::CaptureState(const FBodyInstance* TargetInstance, float GravityZ, float DeltaTime)
 {
 	if (TargetInstance)
@@ -150,7 +137,7 @@ void FVehicleState::CaptureState(const Chaos::FRigidBodyHandle_Internal* Handle,
 /**
  * UChaosVehicleSimulation
  */
-void UChaosVehicleSimulation::TickVehicle(UWorld* WorldIn, float DeltaTime, const FChaosVehicleDefaultAsyncInput& InputData, FChaosVehicleAsyncOutput& OutputData, Chaos::FRigidBodyHandle_Internal* Handle)
+void UChaosVehicleSimulation::TickVehicle(UWorld* WorldIn, float DeltaTime, const FChaosVehicleAsyncInput& InputData, FChaosVehicleAsyncOutput& OutputData, Chaos::FRigidBodyHandle_Internal* Handle)
 {
 	World = WorldIn;
 	RigidHandle = Handle;
@@ -158,11 +145,21 @@ void UChaosVehicleSimulation::TickVehicle(UWorld* WorldIn, float DeltaTime, cons
 	// movement updates and replication
 	if (World && RigidHandle)
 	{
-		if (CanSimulate() && Handle)
+#if DEBUG_NETWORK_PHYSICS
+		if (WorldIn->IsNetMode(NM_ListenServer))
+		{
+			UE_LOG(LogTemp, Log, TEXT("SERVER | PT | TickVehicle | Async tick vehicle with inputs at frame %d : Throttle = %f Brake = %f Roll = %f Pitch = %f Yaw = %f Steering = %f Handbrake = %f Sleeping = %d"),
+				InputData.PhysicsInputs.NetworkInputs.LocalFrame, VehicleInputs.ThrottleInput, VehicleInputs.BrakeInput, VehicleInputs.RollInput, VehicleInputs.PitchInput,
+				VehicleInputs.YawInput, VehicleInputs.SteeringInput, VehicleInputs.HandbrakeInput, VehicleState.bSleeping);
+		}
+		else if(WorldIn->IsNetMode(NM_Client))
 		{
-			// Sync the inputs datas from the history
-			SyncHistoryInputs(InputData, Handle);
+			UE_LOG(LogTemp, Log, TEXT("CLIENT | PT | TickVehicle | Async tick vehicle with inputs at frame %d : Throttle = %f Brake = %f Roll = %f Pitch = %f Yaw = %f Steering = %f Handbrake = %f Sleeping = %d"),
+				InputData.PhysicsInputs.NetworkInputs.LocalFrame, VehicleInputs.ThrottleInput, VehicleInputs.BrakeInput, VehicleInputs.RollInput, VehicleInputs.PitchInput,
+				VehicleInputs.YawInput, VehicleInputs.SteeringInput, VehicleInputs.HandbrakeInput, VehicleState.bSleeping);
 		}
+#endif
+
 		if (!VehicleState.bSleeping)
 		{
 			if (CanSimulate() && Handle)
@@ -171,7 +168,7 @@ void UChaosVehicleSimulation::TickVehicle(UWorld* WorldIn, float DeltaTime, cons
 				UpdateState(DeltaTime, InputData, Handle);
 
 				// Apply the controls inputs
-				ApplyInput(InputData.ControlInputs, DeltaTime);
+				ApplyInput(InputData.PhysicsInputs.NetworkInputs.VehicleInputs, DeltaTime);
 
 				// Update the simulation forces/impulses...
 				UpdateSimulation(DeltaTime, InputData, Handle);
@@ -186,332 +183,17 @@ void UChaosVehicleSimulation::TickVehicle(UWorld* WorldIn, float DeltaTime, cons
 #endif
 }
 
-FVehicleStatesHistory::FVehicleStatesHistory(const int32 FrameCount, const bool bIsHistoryLocal, class UChaosVehicleSimulation* ChaosSimulation)
-	: Super(FrameCount, bIsHistoryLocal), VehicleSimulation(ChaosSimulation)
-{}
-
-FVehicleStatesHistory::~FVehicleStatesHistory()
-{}
-
-bool FVehicleStatesHistory::RecordStates(const int32 RecordFrame, const DatasType& ReplicatedDatas)
-{
-	if (VehicleSimulation && Super::RecordDatas(RecordFrame))
-	{
-		FVehicleStatesDatas& CurrentDatas = GetCurrentDatas();
-		CurrentDatas.StateLastVelocity = ReplicatedDatas.StateLastVelocity;
-		CurrentDatas.EngineOmega = ReplicatedDatas.EngineOmega;
-		CurrentDatas.SuspensionAveragedCount = ReplicatedDatas.SuspensionAveragedCount;
-		CurrentDatas.SuspensionAveragedNum = ReplicatedDatas.SuspensionAveragedNum;
-		CurrentDatas.SuspensionAveragedLength = ReplicatedDatas.SuspensionAveragedLength;
-		CurrentDatas.SuspensionLastDisplacement = ReplicatedDatas.SuspensionLastDisplacement;
-		CurrentDatas.SuspensionLastSpringLength = ReplicatedDatas.SuspensionLastSpringLength;
-		CurrentDatas.WheelsOmega = ReplicatedDatas.WheelsOmega;
-		CurrentDatas.WheelsAngularPosition = ReplicatedDatas.WheelsAngularPosition;
-
-#if !(UE_BUILD_SHIPPING || UE_BUILD_TEST)
-		if (Chaos::FPhysicsSolverBase::CanDebugNetworkPhysicsPrediction())
-		{
-			UE_LOG(LogVehicle, Log, TEXT("Record States at Frame %d on Vehicle = %d with Current = %d | Velocity = %s | Omega = %f"),
-				RecordFrame, VehicleSimulation, CurrentFrame, *CurrentDatas.StateLastVelocity.ToString(), CurrentDatas.EngineOmega);
-		}
-#endif
-		return true;
-	}
-	return false;
-}
-
-void FVehicleStatesHistory::ReplicateStates(DatasType& ReplicatedDatas)
-{
-	using namespace Chaos;
-
-	if (VehicleSimulation)
-	{
-		if (TUniquePtr<Chaos::FSimpleWheeledVehicle>& Vehicle = VehicleSimulation->PVehicle)
-		{
-			ReplicatedDatas.EngineOmega = Vehicle->GetEngine().GetEngineOmega();
-			ReplicatedDatas.StateLastVelocity = VehicleSimulation->VehicleState.LastFrameVehicleLocalVelocity;
-
-			const int32 NumWheels = Vehicle->Wheels.Num();
-
-			int32 NumLength = 0;
-			for (int32 WheelIdx = 0; WheelIdx < NumWheels; ++WheelIdx)
-			{
-				NumLength += Vehicle->GetSuspension(WheelIdx).GetAveragingNum();
-			}
-
-			ReplicatedDatas.SuspensionLastSpringLength.SetNum(NumWheels);
-			ReplicatedDatas.SuspensionLastDisplacement.SetNum(NumWheels);
-			ReplicatedDatas.SuspensionAveragedCount.SetNum(NumWheels);
-			ReplicatedDatas.SuspensionAveragedNum.SetNum(NumWheels);
-			ReplicatedDatas.SuspensionAveragedLength.SetNum(NumLength);
-			ReplicatedDatas.WheelsAngularPosition.SetNum(NumWheels);
-			ReplicatedDatas.WheelsOmega.SetNum(NumWheels);
-
-			int32 LengthCount = 0;
-			for (int32 WheelIdx = 0; WheelIdx < Vehicle->Wheels.Num(); ++WheelIdx)
-			{
-				FSimpleSuspensionSim& Suspension = Vehicle->GetSuspension(WheelIdx);
-				ReplicatedDatas.SuspensionLastSpringLength[WheelIdx] = Suspension.GetLastSpringLength();
-				ReplicatedDatas.SuspensionLastDisplacement[WheelIdx] = Suspension.GetLastDisplacement();
-				ReplicatedDatas.SuspensionAveragedCount[WheelIdx] = Suspension.GetAveragingCount();
-				ReplicatedDatas.SuspensionAveragedNum[WheelIdx] = Suspension.GetAveragingNum();
-
-				for (int32 LengthIdx = 0; LengthIdx < Suspension.GetAveragingNum(); ++LengthIdx)
-				{
-					ReplicatedDatas.SuspensionAveragedLength[LengthCount++] = Suspension.GetAveragingLength(LengthIdx);
-				}
-
-				FSimpleWheelSim& Wheel = Vehicle->GetWheel(WheelIdx);
-				ReplicatedDatas.WheelsOmega[WheelIdx] = Wheel.Omega;
-				ReplicatedDatas.WheelsAngularPosition[WheelIdx] = Wheel.AngularPosition;
-			}
-		}
-	}
-}
-
-bool FVehicleStatesHistory::RewindStates(const int32 RewindFrame, const bool bResetSolver)
-{
-	using namespace Chaos;
-
-	if (VehicleSimulation && Super::RewindStates(RewindFrame, bResetSolver))
-	{
-		const FVehicleStatesDatas& CurrentDatas = GetCurrentDatas();
-		VehicleSimulation->VehicleState.LastFrameVehicleLocalVelocity = CurrentDatas.StateLastVelocity;
-
-		if (TUniquePtr<Chaos::FSimpleWheeledVehicle>& Vehicle = VehicleSimulation->PVehicle)
-		{
-			Vehicle->GetEngine().SetEngineOmega(CurrentDatas.EngineOmega);
-
-			int32 LengthCount = 0;
-			for (int32 WheelIdx = 0, NumWheels = Vehicle->Wheels.Num(); WheelIdx < NumWheels; ++WheelIdx)
-			{
-				FSimpleSuspensionSim& Suspension = Vehicle->GetSuspension(WheelIdx);
-				Suspension.SetLastSpringLength(CurrentDatas.SuspensionLastSpringLength[WheelIdx]);
-				Suspension.SetLastDisplacement(CurrentDatas.SuspensionLastDisplacement[WheelIdx]);
-				Suspension.SetAveragingCount(CurrentDatas.SuspensionAveragedCount[WheelIdx]);
-				Suspension.SetAveragingNum(CurrentDatas.SuspensionAveragedNum[WheelIdx]);
-
-				for (int32 LengthIdx = 0; LengthIdx < Suspension.GetAveragingNum(); ++LengthIdx)
-				{
-					Suspension.SetAveragingLength(LengthIdx, CurrentDatas.SuspensionAveragedLength[LengthCount++]);
-				}
-
-				FSimpleWheelSim& Wheel = Vehicle->GetWheel(WheelIdx);
-				Wheel.Omega = CurrentDatas.WheelsOmega[WheelIdx];
-				Wheel.AngularPosition = CurrentDatas.WheelsAngularPosition[WheelIdx];
-			}
-#if !(UE_BUILD_SHIPPING || UE_BUILD_TEST)
-			if (Chaos::FPhysicsSolverBase::CanDebugNetworkPhysicsPrediction())
-			{
-				UE_LOG(LogVehicle, Log, TEXT("Rewind States at Frame %d on Vehicle = %d with Current = %d | Velocity = %s | Omega = %f"),
-					RewindFrame, VehicleSimulation, CurrentFrame, *CurrentDatas.StateLastVelocity.ToString(), CurrentDatas.EngineOmega);
-			}
-#endif
-		}
-		return true;
-	}
-	return false;
-}
-
-FVehicleInputsHistory::FVehicleInputsHistory(const int32 FrameCount, const bool bIsHistoryLocal, class UChaosVehicleSimulation* ChaosSimulation)
-	: Super(FrameCount, bIsHistoryLocal), VehicleSimulation(ChaosSimulation)
-{}
-
-FVehicleInputsHistory::~FVehicleInputsHistory()
-{}
-
-bool FVehicleInputsHistory::RecordInputs(const int32 RecordFrame, const DatasType& ReplicatedDatas)
-{
-	if (VehicleSimulation && Super::RecordDatas(RecordFrame))
-	{
-		FVehicleInputsDatas& CurrentDatas = GetCurrentDatas();
-		CurrentDatas.VehicleInputs.SteeringInput = ReplicatedDatas.VehicleInputs.SteeringInput;
-		CurrentDatas.VehicleInputs.ThrottleInput = ReplicatedDatas.VehicleInputs.ThrottleInput;
-		CurrentDatas.VehicleInputs.BrakeInput = ReplicatedDatas.VehicleInputs.BrakeInput;
-		CurrentDatas.VehicleInputs.HandbrakeInput = ReplicatedDatas.VehicleInputs.HandbrakeInput;
-		CurrentDatas.VehicleInputs.RollInput = ReplicatedDatas.VehicleInputs.RollInput;
-		CurrentDatas.VehicleInputs.PitchInput = ReplicatedDatas.VehicleInputs.PitchInput;
-		CurrentDatas.VehicleInputs.YawInput = ReplicatedDatas.VehicleInputs.YawInput;
-		CurrentDatas.TransmissionCurrentGear = ReplicatedDatas.TransmissionCurrentGear;
-		CurrentDatas.TransmissionTargetGear = ReplicatedDatas.TransmissionTargetGear;
-		CurrentDatas.TransmissionChangeTime = ReplicatedDatas.TransmissionChangeTime;
-
-#if !(UE_BUILD_SHIPPING || UE_BUILD_TEST)
-		if (Chaos::FPhysicsSolverBase::CanDebugNetworkPhysicsPrediction())
-		{
-			UE_LOG(LogVehicle, Log, TEXT("Record Inputs at Frame %d on Vehicle = %d with Current = %d | Brake = %f | Handbrake = %f | Pitch = %f | Roll = %f | Steering = %f | Throttle = %f | Yaw = %f | Gear = %d"),
-				RecordFrame, VehicleSimulation, CurrentFrame, CurrentDatas.VehicleInputs.BrakeInput, CurrentDatas.VehicleInputs.HandbrakeInput,
-				CurrentDatas.VehicleInputs.PitchInput, CurrentDatas.VehicleInputs.RollInput, CurrentDatas.VehicleInputs.SteeringInput,
-				CurrentDatas.VehicleInputs.ThrottleInput, CurrentDatas.VehicleInputs.YawInput, CurrentDatas.TransmissionCurrentGear);
-		}
-#endif
-		return true;
-	}
-	return false;
-}
-
-void FVehicleInputsHistory::ReplicateInputs(
-	const float SteeringInput, const float ThrottleInput,
-	const float BrakeInput, const float HandbrakeInput, const float RollInput,
-	const float PitchInput, const float YawInput, const int32 CurrentGear, DatasType& ReplicatedDatas)
-{
-	using namespace Chaos;
-
-	ReplicatedDatas.VehicleInputs.SteeringInput = SteeringInput;
-	ReplicatedDatas.VehicleInputs.ThrottleInput = ThrottleInput;
-	ReplicatedDatas.VehicleInputs.BrakeInput = BrakeInput;
-	ReplicatedDatas.VehicleInputs.HandbrakeInput = HandbrakeInput;
-	ReplicatedDatas.VehicleInputs.RollInput = RollInput;
-	ReplicatedDatas.VehicleInputs.PitchInput = PitchInput;
-	ReplicatedDatas.VehicleInputs.YawInput = YawInput;
-
-	if (VehicleSimulation)
-	{
-		if (TUniquePtr<Chaos::FSimpleWheeledVehicle>& Vehicle = VehicleSimulation->PVehicle)
-		{
-			if (Vehicle->HasTransmission())
-			{
-				FSimpleTransmissionSim& Transmission = Vehicle->GetTransmission();
-
-				Transmission.SetGear(CurrentGear, true);
-				ReplicatedDatas.TransmissionCurrentGear = Transmission.GetCurrentGear();
-				ReplicatedDatas.TransmissionTargetGear = Transmission.GetCurrentGear();
-				ReplicatedDatas.TransmissionChangeTime = Transmission.GetCurrentGearChangeTime();
-			}
-		}
-	}
-}
-
-void FVehicleInputsHistory::SmoothInputs(const int32 ApplyFrame, const int32 ValidFrame, const bool bApplySmoothing)
-{
-	using namespace Chaos;
-
-	if (!bIsLocalHistory && VehicleSimulation)
-	{
-		const float HalfFrames = (float)NumFrames / 2.0f;
-		const float GaussSigma = HalfFrames / 2.0f;
-
-		const int32 ValidIndex = ValidFrame % NumFrames;
-		const int32 FrameIndex = ValidFrame - ApplyFrame;
-
-		const float SmoothFactor = bApplySmoothing ? FMath::Exp(-FrameIndex * FrameIndex / (2.0f * GaussSigma * GaussSigma)) : 1.0;
-
-		VehicleSimulation->VehicleInputs.BrakeInput = SmoothFactor * DatasArray[ValidIndex].VehicleInputs.BrakeInput;
-		VehicleSimulation->VehicleInputs.HandbrakeInput = SmoothFactor * DatasArray[ValidIndex].VehicleInputs.HandbrakeInput;
-		VehicleSimulation->VehicleInputs.PitchInput = SmoothFactor * DatasArray[ValidIndex].VehicleInputs.PitchInput;
-		VehicleSimulation->VehicleInputs.RollInput = SmoothFactor * DatasArray[ValidIndex].VehicleInputs.RollInput;
-		VehicleSimulation->VehicleInputs.ThrottleInput = SmoothFactor * DatasArray[ValidIndex].VehicleInputs.ThrottleInput;
-		VehicleSimulation->VehicleInputs.SteeringInput = SmoothFactor * DatasArray[ValidIndex].VehicleInputs.SteeringInput;
-		VehicleSimulation->VehicleInputs.YawInput = SmoothFactor * DatasArray[ValidIndex].VehicleInputs.YawInput;
-
-		if (TUniquePtr<Chaos::FSimpleWheeledVehicle>& Vehicle = VehicleSimulation->PVehicle)
-		{
-			if (Vehicle->HasTransmission())
-			{
-				FSimpleTransmissionSim& Transmission = VehicleSimulation->PVehicle->GetTransmission();
-				Transmission.SetCurrentGear(DatasArray[ValidIndex].TransmissionCurrentGear);
-				Transmission.SetTargetGear(DatasArray[ValidIndex].TransmissionTargetGear);
-				Transmission.SetCurrentGearChangeTime(DatasArray[ValidIndex].TransmissionChangeTime);
-			}
-
-#if !(UE_BUILD_SHIPPING || UE_BUILD_TEST)
-			if (Chaos::FPhysicsSolverBase::CanDebugNetworkPhysicsPrediction())
-			{
-				UE_LOG(LogVehicle, Log, TEXT("Smooth Inputs from frame %d to frame = %d with factor %f"),
-					ValidFrame, ApplyFrame, SmoothFactor);
-			}
-#endif
-		}
-	}
-}
-
-bool FVehicleInputsHistory::ApplyInputs(const int32 ApplyFrame, const bool bResetSolver)
-{
-	using namespace Chaos;
-
-	if (VehicleSimulation )
-	{
-		if(Super::ApplyInputs(ApplyFrame, bResetSolver))
-		{
-			const FVehicleInputsDatas& CurrentDatas = GetCurrentDatas();
-
-			VehicleSimulation->VehicleInputs = CurrentDatas.VehicleInputs;
-
-			FSimpleTransmissionSim& Transmission = VehicleSimulation->PVehicle->GetTransmission();
-			Transmission.SetCurrentGear(CurrentDatas.TransmissionCurrentGear);
-			Transmission.SetTargetGear(CurrentDatas.TransmissionTargetGear);
-			Transmission.SetCurrentGearChangeTime(CurrentDatas.TransmissionChangeTime);
-		}
-		else
-		{
-			int32 ClosestFrame = INDEX_NONE;
-			for (int32 FrameIndex = 0; FrameIndex < NumFrames; ++FrameIndex)
-			{
-				const int32 ValidFrame = FMath::Max(0, ApplyFrame - FrameIndex);
-				const int32 ValidIndex = ValidFrame % NumFrames;
-
-				if (DatasArray[ValidIndex].DatasFrame == ValidFrame)
-				{
-					ClosestFrame = ValidFrame;
-					break;
-				}
-			}
-			if(ClosestFrame != INDEX_NONE)
-			{ 
-				SmoothInputs(ApplyFrame, ClosestFrame, false);
-			}
-		}
-
-#if !(UE_BUILD_SHIPPING || UE_BUILD_TEST)
-		FSimpleTransmissionSim& Transmission = VehicleSimulation->PVehicle->GetTransmission();
-
-		if (Chaos::FPhysicsSolverBase::CanDebugNetworkPhysicsPrediction())
-		{
-			UE_LOG(LogVehicle, Log, TEXT("Apply Inputs at Frame %d on Vehicle = %d with Current = %d | Brake = %f | Handbrake = %f | Pitch = %f | Roll = %f | Steering = %f | Throttle = %f | Yaw = %f | Gear = %d"),
-				ApplyFrame, VehicleSimulation, CurrentFrame, VehicleSimulation->VehicleInputs.BrakeInput,
-				VehicleSimulation->VehicleInputs.HandbrakeInput, VehicleSimulation->VehicleInputs.PitchInput,
-				VehicleSimulation->VehicleInputs.RollInput, VehicleSimulation->VehicleInputs.SteeringInput,
-				VehicleSimulation->VehicleInputs.ThrottleInput, VehicleSimulation->VehicleInputs.YawInput,
-				Transmission.GetCurrentGear());
-
-		}
-#endif
-		return true;
-	}
-	return false;
-}
-
-void UChaosVehicleSimulation::SyncHistoryInputs(const FChaosVehicleDefaultAsyncInput& InputData, Chaos::FRigidBodyHandle_Internal* Handle)
-{
-	if (Chaos::FPhysicsSolverBase::IsNetworkPhysicsPredictionEnabled())
-	{
-		InputData.ControlInputs.CopyVehicleInputs(VehicleInputs);
-
-#if !(UE_BUILD_SHIPPING || UE_BUILD_TEST)
-		if (Chaos::FPhysicsSolverBase::CanDebugNetworkPhysicsPrediction())
-		{
-			auto* RigidSolver = Handle->GetProxy()->GetSolver<Chaos::FPBDRigidsSolver>();
-			if (RigidSolver && !RigidSolver->GetEvolution()->IsResimming())
-			{
-				UE_LOG(LogVehicle, Log, TEXT("Store States at Frame %d on Vehicle = %d | Velocity = %s | Omega = %f"),
-					RigidSolver->GetCurrentFrame(), this, *VehicleState.LastFrameVehicleLocalVelocity.ToString(), PVehicle->GetEngine().GetEngineOmega());
-			}
-		}
-#endif
-	}
-}
-
 void UChaosVehicleSimulation::ApplyDeferredForces(Chaos::FRigidBodyHandle_Internal* Handle)
 {
 	DeferredForces.Apply(Handle);
 }
 
-void UChaosVehicleSimulation::UpdateState(float DeltaTime, const FChaosVehicleDefaultAsyncInput& InputData, Chaos::FRigidBodyHandle_Internal* Handle)
+void UChaosVehicleSimulation::UpdateState(float DeltaTime, const FChaosVehicleAsyncInput& InputData, Chaos::FRigidBodyHandle_Internal* Handle)
 {
-	VehicleState.CaptureState(Handle, InputData.GravityZ, DeltaTime);
+	VehicleState.CaptureState(Handle, InputData.PhysicsInputs.GravityZ, DeltaTime);
 }
 
-void UChaosVehicleSimulation::UpdateSimulation(float DeltaTime, const FChaosVehicleDefaultAsyncInput& InputData, Chaos::FRigidBodyHandle_Internal* Handle)
+void UChaosVehicleSimulation::UpdateSimulation(float DeltaTime, const FChaosVehicleAsyncInput& InputData, Chaos::FRigidBodyHandle_Internal* Handle)
 {
 	ApplyAerodynamics(DeltaTime);
 	ApplyAerofoilForces(DeltaTime);
@@ -692,14 +374,14 @@ void UChaosVehicleSimulation::ApplyThrustForces(float DeltaTime)
 }
 
 
-void UChaosVehicleSimulation::ApplyTorqueControl(float DeltaTime, const FChaosVehicleDefaultAsyncInput& InputData)
+void UChaosVehicleSimulation::ApplyTorqueControl(float DeltaTime, const FChaosVehicleAsyncInput& InputData)
 {
 	if (!PVehicle->HasTorqueControlSetup())
 	{
 		return;
 	}
 
-	const FControlInputs& ControlInputsPT = InputData.ControlInputs;
+	const FControlInputs& ControlInputsPT = InputData.PhysicsInputs.NetworkInputs.VehicleInputs;
 
 	if (!GVehicleDebugParams.DisableTorqueControl && RigidHandle)
 	{
@@ -966,6 +648,15 @@ UChaosVehicleMovementComponent::UChaosVehicleMovementComponent(const FObjectInit
 	bUsingNetworkPhysicsPrediction = Chaos::FPhysicsSolverBase::IsNetworkPhysicsPredictionEnabled();
 
 	AHUD::OnShowDebugInfo.AddUObject(this, &UChaosVehicleMovementComponent::ShowDebugInfo);
+
+	if (bUsingNetworkPhysicsPrediction)
+	{
+		static const FName NetworkPhysicsComponentName(TEXT("PC_NetworkPhysicsComponent"));
+
+		NetworkPhysicsComponent = CreateDefaultSubobject<UNetworkPhysicsComponent, UNetworkPhysicsComponent>(NetworkPhysicsComponentName);
+		NetworkPhysicsComponent->SetNetAddressable(); // Make DSO components net addressable
+		NetworkPhysicsComponent->SetIsReplicated(true);
+	}
 }
 
 // public
@@ -1004,7 +695,6 @@ void UChaosVehicleMovementComponent::SetOverrideController(AController* InOverri
 	OverrideController = InOverrideController;
 }
 
-
 bool UChaosVehicleMovementComponent::ShouldCreatePhysicsState() const
 {
 	if (!IsRegistered() || IsBeingDestroyed())
@@ -1055,7 +745,6 @@ bool UChaosVehicleMovementComponent::CanCreateVehicle() const
 	return true;
 }
 
-
 void UChaosVehicleMovementComponent::OnCreatePhysicsState()
 {
 	Super::OnCreatePhysicsState();
@@ -1081,24 +770,9 @@ void UChaosVehicleMovementComponent::OnCreatePhysicsState()
 			}
 			if (bUsingNetworkPhysicsPrediction)
 			{
-				if (Chaos::FPhysicsSolver* Solver = PhysScene->GetSolver())
+				if(NetworkPhysicsComponent)
 				{
-					if (Chaos::FRewindData* RewindData = Solver->GetRewindData())
-					{
-						if (VehicleSimulationPT.IsValid())
-						{
-							AController* Controller = GetController();
-							const bool bIsLocalHistory = (Controller && Controller->IsLocalController());
-
-							const int32 PhysicsHistoryLength = UPhysicsSettings::Get()->GetPhysicsHistoryCount();
-
-							VehicleSimulationPT->InputsHistory = MakeShared<FVehicleInputsHistory>(PhysicsHistoryLength, bIsLocalHistory, VehicleSimulationPT.Get());
-							RewindData->AddInputsHistory(VehicleSimulationPT->InputsHistory);
-
-							VehicleSimulationPT->StatesHistory = MakeShared<FVehicleStatesHistory>(PhysicsHistoryLength, bIsLocalHistory, VehicleSimulationPT.Get());
-							RewindData->AddStatesHistory(VehicleSimulationPT->StatesHistory);
-						}
-					}
+					NetworkPhysicsComponent->CreateDatasHistory<FPhysicsVehicleTraits>(this);
 				}
 			}
 		}
@@ -1127,26 +801,9 @@ void UChaosVehicleMovementComponent::OnDestroyPhysicsState()
 			UpdatedComponent->RecreatePhysicsState();
 		}
 	}
-	if (bUsingNetworkPhysicsPrediction)
+	if (bUsingNetworkPhysicsPrediction && NetworkPhysicsComponent)
 	{
-		UWorld* World = GetWorld();
-		if (World->IsGameWorld())
-		{
-			if (FPhysScene* PhysScene = World->GetPhysicsScene())
-			{
-				if (Chaos::FPhysicsSolver* Solver = PhysScene->GetSolver())
-				{
-					if (Chaos::FRewindData* RewindData = Solver->GetRewindData())
-					{
-						if (VehicleSimulationPT.IsValid())
-						{
-							RewindData->RemoveInputsHistory(VehicleSimulationPT->InputsHistory);
-							RewindData->RemoveStatesHistory(VehicleSimulationPT->StatesHistory);
-						}
-					}
-				}
-			}
-		}
+		NetworkPhysicsComponent->RemoveDatasHistory();
 	}
 }
 
@@ -1175,6 +832,10 @@ void UChaosVehicleMovementComponent::PreTickGT(float DeltaTime)
 		ControlInputs.ParkingEnabled = bParkEnabled;
 		ProcessSleeping(ControlInputs);
 	}
+	else
+	{
+		ProcessSleeping(VehicleSimulationPT->VehicleInputs);
+	}
 
 	if (VehicleSetupTag != FChaosVehicleManager::VehicleSetupTag)
 	{
@@ -1286,20 +947,17 @@ void UChaosVehicleMovementComponent::SetTargetGear(int32 GearNum, bool bImmediat
 {
 	if (PVehicleOutput && GearNum != PVehicleOutput->TargetGear)
 	{
-		if (!bUsingNetworkPhysicsPrediction)
-		{
-			FBodyInstance* TargetInstance = UpdatedPrimitive->GetBodyInstance();
+		FBodyInstance* TargetInstance = UpdatedPrimitive->GetBodyInstance();
 
-			if (TargetInstance)
-			{
-				FPhysicsCommand::ExecuteWrite(TargetInstance->ActorHandle, [&](const FPhysicsActorHandle& Chassis)
+		if (TargetInstance)
+		{
+			FPhysicsCommand::ExecuteWrite(TargetInstance->ActorHandle, [&](const FPhysicsActorHandle& Chassis)
+				{
+					if (VehicleSimulationPT && VehicleSimulationPT->PVehicle && VehicleSimulationPT->PVehicle->HasTransmission())
 					{
-						if (VehicleSimulationPT && VehicleSimulationPT->PVehicle && VehicleSimulationPT->PVehicle->HasTransmission())
-						{
-							VehicleSimulationPT->PVehicle->GetTransmission().SetGear(GearNum, bImmediate);
-						}
-					});
-			}
+						VehicleSimulationPT->PVehicle->GetTransmission().SetGear(GearNum, bImmediate);
+					}
+				});
 		}
 
 		TargetGear = GearNum;
@@ -1566,25 +1224,6 @@ void UChaosVehicleMovementComponent::UpdateState(float DeltaTime)
 			// and send to server - (ServerUpdateState_Implementation below)
 			ServerUpdateState(SteeringInput, ThrottleInput, BrakeInput, HandbrakeInput, GetCurrentGear(), RollInput, PitchInput, YawInput);
 		}
-		else
-		{
-			if (APawn* MyPawn = Cast<APawn>(GetOwner()))
-			{
-				if (APlayerController* PlayerController = Cast<APlayerController>(MyPawn->GetController()))
-				{
-					if (PlayerController->IsLocalController() && !PlayerController->IsNetMode(NM_DedicatedServer))
-					{
-						const FAsyncPhysicsTimestamp TimeStamp = PlayerController->GetAsyncPhysicsTimestamp();
-						ApplyControls_Imp(SteeringInput, ThrottleInput, BrakeInput, HandbrakeInput, TargetGear, RollInput, PitchInput, YawInput, TimeStamp, true);
-
-						if (!PlayerController->IsNetMode(NM_ListenServer))
-						{
-							ServerApplyControls(SteeringInput, ThrottleInput, BrakeInput, HandbrakeInput, TargetGear, RollInput, PitchInput, YawInput, TimeStamp);
-						}
-					}
-				}
-			}
-		}
 
 		if (PawnOwner && PawnOwner->IsNetMode(NM_Client))
 		{
@@ -1628,7 +1267,7 @@ void UChaosVehicleMovementComponent::ProcessSleeping(const FControlInputs& Contr
 		//			be processed locally). Maybe we should *just* use ReplicatedState?
 
 		const AController* Controller = GetController();
-		const bool bIsLocallyControlled = (Controller && Controller->IsLocalController());
+		const bool bIsLocallyControlled = (Controller && Controller->IsLocalController()) || bUsingNetworkPhysicsPrediction;
 		const bool bControlInputPressed = bIsLocallyControlled ? (ControlInputs.ThrottleInput >= GVehicleDebugParams.ControlInputWakeTolerance)
 			|| (ControlInputs.BrakeInput >= GVehicleDebugParams.ControlInputWakeTolerance)
 			|| (FMath::Abs(ControlInputs.SteeringInput - PrevSteeringInput) >= GVehicleDebugParams.ControlInputWakeTolerance)
@@ -1712,57 +1351,6 @@ void UChaosVehicleMovementComponent::ServerUpdateState_Implementation(float InSt
 
 }
 
-void UChaosVehicleMovementComponent::ServerApplyControls_Implementation(float InSteeringInput, float InThrottleInput, float InBrakeInput, float InHandbrakeInput,
-	int32 InCurrentGear, float InRollInput, float InPitchInput, float InYawInput, const FAsyncPhysicsTimestamp& AsyncPhysicsTimestamp)
-{
-	ApplyControls_Imp(InSteeringInput, InThrottleInput, InBrakeInput, InHandbrakeInput, InCurrentGear, InRollInput, InPitchInput, InYawInput, AsyncPhysicsTimestamp, false);
-}
-
-void UChaosVehicleMovementComponent::ApplyControls_Imp(float InSteeringInput, float InThrottleInput, float InBrakeInput, float InHandbrakeInput,
-	int32 InCurrentGear, float InRollInput, float InPitchInput, float InYawInput, const FAsyncPhysicsTimestamp& AsyncPhysicsTimestamp, const bool bEnableResim)
-{
-	if (UWorld* World = GetWorld())
-	{
-		APawn* LocalPawn = Cast<APawn>(GetOwner());
-		if (APlayerController* PlayerController = (LocalPawn && LocalPawn->GetController()) ?
-			Cast<APlayerController>(LocalPawn->GetController()) : World->GetFirstPlayerController())
-		{
-			const bool bControlInputPressed = ((InThrottleInput >= GVehicleDebugParams.ControlInputWakeTolerance)
-				|| (InBrakeInput >= GVehicleDebugParams.ControlInputWakeTolerance)
-				|| (InRollInput >= GVehicleDebugParams.ControlInputWakeTolerance)
-				|| (InPitchInput >= GVehicleDebugParams.ControlInputWakeTolerance)
-				|| (InYawInput >= GVehicleDebugParams.ControlInputWakeTolerance)
-				|| (FMath::Abs(InSteeringInput - PrevSteeringInput) >= GVehicleDebugParams.ControlInputWakeTolerance));
-
-			if (bControlInputPressed)
-			{
-				SetSleeping(false);
-			}
-			const FAsyncPhysicsTimestamp LocalStep = PlayerController->GetAsyncPhysicsTimestamp();
-
-			PlayerController->ExecuteAsyncPhysicsCommand(AsyncPhysicsTimestamp, this, [this, InSteeringInput, InThrottleInput, InBrakeInput, InHandbrakeInput, InCurrentGear, InRollInput, InPitchInput, InYawInput, AsyncPhysicsTimestamp, LocalStep]
-			{
-				const bool bIsServer = (GetWorld()->IsNetMode(NM_ListenServer) || GetWorld()->IsNetMode(NM_DedicatedServer));
-				const int32 HistoryFrame = bIsServer ? AsyncPhysicsTimestamp.ServerFrame : AsyncPhysicsTimestamp.LocalFrame;
-
-				ReplicatedDatas.ServerFrame = AsyncPhysicsTimestamp.ServerFrame;
-
-				// Replicate the Inputs to send them to other clients
-				VehicleSimulationPT->InputsHistory->ReplicateInputs(InSteeringInput, InThrottleInput, 
-					InBrakeInput, InHandbrakeInput, InRollInput, InPitchInput, InYawInput, InCurrentGear, ReplicatedDatas.InputsDatas);
-
-				// We record in the history buffer the local client control inputs since they are the ground truth
-				VehicleSimulationPT->InputsHistory->RecordInputs(HistoryFrame, ReplicatedDatas.InputsDatas);
-					
-				// On the server we replicate the vehicle state in order to be used during resimulation at the rewind step
-				if (bIsServer)
-				{
-					VehicleSimulationPT->StatesHistory->ReplicateStates(ReplicatedDatas.StatesDatas);
-				}
-			}, bEnableResim);
-		}
-	}
-}
 /// @endcond
 
 
@@ -1785,6 +1373,15 @@ AController* UChaosVehicleMovementComponent::GetController() const
 	return nullptr;
 }
 
+APlayerController* UChaosVehicleMovementComponent::GetPlayerController() const
+{
+	if (AController* Controller = GetController())
+	{
+		return Cast<APlayerController>(Controller);
+	}
+	return nullptr;
+}
+
 
 FBodyInstance* UChaosVehicleMovementComponent::GetBodyInstance()
 {
@@ -2006,7 +1603,6 @@ void UChaosVehicleMovementComponent::GetLifetimeReplicatedProps(TArray< FLifetim
 
 	DOREPLIFETIME(UChaosVehicleMovementComponent, ReplicatedState);
 	DOREPLIFETIME(UChaosVehicleMovementComponent, OverrideController);
-	DOREPLIFETIME(UChaosVehicleMovementComponent, ReplicatedDatas);
 }
 
 /// @endcond
@@ -2056,7 +1652,7 @@ void FVehicleThrustConfig::FillThrusterSetup(const UChaosVehicleMovementComponen
 
 TUniquePtr<FChaosVehicleAsyncInput> UChaosVehicleMovementComponent::SetCurrentAsyncInputOutput(int32 InputIdx, FChaosVehicleManagerAsyncOutput* CurOutput, FChaosVehicleManagerAsyncOutput* NextOutput, float Alpha, int32 VehicleManagerTimestamp)
 {
-	TUniquePtr<FChaosVehicleDefaultAsyncInput> CurInput = MakeUnique<FChaosVehicleDefaultAsyncInput>();
+	TUniquePtr<FChaosVehicleAsyncInput> CurInput = MakeUnique<FChaosVehicleAsyncInput>();
 	SetCurrentAsyncInputOutputInternal(CurInput.Get(), InputIdx, CurOutput, NextOutput, Alpha, VehicleManagerTimestamp);
 	return CurInput;
 }
@@ -2146,46 +1742,6 @@ void UChaosVehicleMovementComponent::SetCurrentAsyncInputOutputInternal(FChaosVe
 
 // ---- ASYNC ----
 
-FChaosVehicleDefaultAsyncInput::FChaosVehicleDefaultAsyncInput()
-	: GravityZ(0.f)
-{
-
-}
-
-
-void FChaosVehicleDefaultAsyncInput::ApplyDeferredForces(Chaos::FRigidBodyHandle_Internal* RigidHandle) const
-{
-	check(Vehicle);
-	check(Vehicle->VehicleSimulationPT);
-	Vehicle->VehicleSimulationPT->ApplyDeferredForces(RigidHandle);
-}
-
-/************************************************************************/
-/* Async simulation callback on the Physics Thread                      */
-/************************************************************************/
-TUniquePtr<FChaosVehicleAsyncOutput> FChaosVehicleDefaultAsyncInput::Simulate(UWorld* World, const float DeltaSeconds, const float TotalSeconds, bool& bWakeOut) const
-{
-	TUniquePtr<FChaosVehicleAsyncOutput> Output = MakeUnique<FChaosVehicleAsyncOutput>();
-
-	//UE_LOG(LogChaos, Warning, TEXT("Vehicle Physics Thread Tick %f"), DeltaSeconds);
-
-	//support nullptr because it allows us to go wide on filling the async inputs
-	if (Proxy == nullptr)
-	{
-		return Output;
-	}
-
-	// We now have access to the physics representation of the chassis on the physics thread async tick
-	Chaos::FRigidBodyHandle_Internal* Handle = Proxy->GetPhysicsThreadAPI();
-
-	// FILL OUTPUT DATA HERE THAT WILL GET PASSED BACK TO THE GAME THREAD
-	Vehicle->VehicleSimulationPT->TickVehicle(World, DeltaSeconds, *this, *Output.Get(), Handle);
-
-	Output->bValid = true;
-
-	return MoveTemp(Output);
-}
-
 /************************************************************************/
 /* PASS ANY INUTS TO THE PHYSICS THREAD SIMULATION IN HERE              */
 /************************************************************************/
@@ -2198,32 +1754,27 @@ void UChaosVehicleMovementComponent::Update(float DeltaTime)
 			if (auto Handle = BodyInstance->ActorHandle)
 			{
 				CurAsyncInput->Proxy = Handle;	// vehicles are never static
-				FChaosVehicleDefaultAsyncInput* AsyncInput = static_cast<FChaosVehicleDefaultAsyncInput*>(CurAsyncInput);
-
-				if (!bUsingNetworkPhysicsPrediction)
+				FChaosVehicleAsyncInput* AsyncInput = static_cast<FChaosVehicleAsyncInput*>(CurAsyncInput);
+
+				AsyncInput->PhysicsInputs.NetworkInputs.VehicleInputs.ThrottleInput = ThrottleInputRate.CalcControlFunction(ThrottleInput);
+				AsyncInput->PhysicsInputs.NetworkInputs.VehicleInputs.BrakeInput = BrakeInputRate.CalcControlFunction(BrakeInput);
+				AsyncInput->PhysicsInputs.NetworkInputs.VehicleInputs.SteeringInput = SteeringInputRate.CalcControlFunction(SteeringInput);
+				AsyncInput->PhysicsInputs.NetworkInputs.VehicleInputs.HandbrakeInput = HandbrakeInput;
+				AsyncInput->PhysicsInputs.NetworkInputs.VehicleInputs.RollInput = RollInputRate.CalcControlFunction(RollInput);
+				AsyncInput->PhysicsInputs.NetworkInputs.VehicleInputs.PitchInput = PitchInputRate.CalcControlFunction(PitchInput);
+				AsyncInput->PhysicsInputs.NetworkInputs.VehicleInputs.YawInput = YawInputRate.CalcControlFunction(YawInput);
+				AsyncInput->PhysicsInputs.NetworkInputs.VehicleInputs.GearUpInput = bRawGearUpInput;
+				AsyncInput->PhysicsInputs.NetworkInputs.VehicleInputs.GearDownInput = bRawGearDownInput;
+				AsyncInput->PhysicsInputs.NetworkInputs.VehicleInputs.TransmissionType = TransmissionType;
+				AsyncInput->PhysicsInputs.NetworkInputs.VehicleInputs.ParkingEnabled = bParkEnabled;
+
+				// debug feature to limit the vehicles top speed
+				if ((GVehicleDebugParams.SetMaxMPH > 0.f) && (FMath::Abs(ThrottleInput) > 0.0f) && FMath::Abs(GetForwardSpeedMPH()) >= GVehicleDebugParams.SetMaxMPH)
 				{
-					AsyncInput->ControlInputs.ThrottleInput = ThrottleInputRate.CalcControlFunction(ThrottleInput);
-					AsyncInput->ControlInputs.BrakeInput = BrakeInputRate.CalcControlFunction(BrakeInput);
-					AsyncInput->ControlInputs.SteeringInput = SteeringInputRate.CalcControlFunction(SteeringInput);
-					AsyncInput->ControlInputs.HandbrakeInput = HandbrakeInput;
-					AsyncInput->ControlInputs.RollInput = RollInputRate.CalcControlFunction(RollInput);
-					AsyncInput->ControlInputs.PitchInput = PitchInputRate.CalcControlFunction(PitchInput);
-					AsyncInput->ControlInputs.YawInput = YawInputRate.CalcControlFunction(YawInput);
-					AsyncInput->ControlInputs.GearUpInput = bRawGearUpInput;
-					AsyncInput->ControlInputs.GearDownInput = bRawGearDownInput;
-					AsyncInput->ControlInputs.TransmissionType = TransmissionType;
-
-					AsyncInput->ControlInputs.ParkingEnabled = bParkEnabled;
-
-					// debug feature to limit the vehicles top speed
-					if ((GVehicleDebugParams.SetMaxMPH > 0.f) && (FMath::Abs(ThrottleInput) > 0.0f) && FMath::Abs(GetForwardSpeedMPH()) >= GVehicleDebugParams.SetMaxMPH)
-					{
-						AsyncInput->ControlInputs.ThrottleInput = 0.1f;
-					}
+					AsyncInput->PhysicsInputs.NetworkInputs.VehicleInputs.ThrottleInput = 0.1f;
 				}
 
-				AsyncInput->GravityZ = GetGravityZ();
-
+				AsyncInput->PhysicsInputs.GravityZ = GetGravityZ();
 			}
 		}
 	}
@@ -2404,58 +1955,6 @@ void UChaosVehicleMovementComponent::PutAllEnabledRigidBodiesToSleep()
 	}
 }
 
-void UChaosVehicleMovementComponent::OnRep_SetVehicleHistory()
-{
-	if (UWorld* World = GetWorld())
-	{
-		bool bHasController = false;
-		if (APawn* MyPawn = Cast<APawn>(GetOwner()))
-		{
-			if (APlayerController* PlayerController = Cast<APlayerController>(MyPawn->GetController()))
-			{
-				if (PlayerController->IsLocalController() && !PlayerController->IsNetMode(NM_DedicatedServer))
-				{
-					bHasController = true;
-				}
-			}
-		}
-
-		if (!bHasController)
-		{
-			const bool bControlInputPressed = ((ReplicatedDatas.InputsDatas.VehicleInputs.ThrottleInput >= GVehicleDebugParams.ControlInputWakeTolerance)
-				|| (ReplicatedDatas.InputsDatas.VehicleInputs.BrakeInput >= GVehicleDebugParams.ControlInputWakeTolerance)
-				|| (ReplicatedDatas.InputsDatas.VehicleInputs.RollInput >= GVehicleDebugParams.ControlInputWakeTolerance)
-				|| (ReplicatedDatas.InputsDatas.VehicleInputs.PitchInput >= GVehicleDebugParams.ControlInputWakeTolerance)
-				|| (ReplicatedDatas.InputsDatas.VehicleInputs.YawInput >= GVehicleDebugParams.ControlInputWakeTolerance)
-				|| (FMath::Abs(ReplicatedDatas.InputsDatas.VehicleInputs.SteeringInput - PrevSteeringInput) >= GVehicleDebugParams.ControlInputWakeTolerance));
-
-			if (bControlInputPressed)
-			{
-				SetSleeping(false);
-			}
-		}
-
-		APawn* LocalPawn = Cast<APawn>(GetOwner());
-		if (APlayerController* PlayerController = (LocalPawn && LocalPawn->GetController()) ?
-			Cast<APlayerController>(LocalPawn->GetController()) : World->GetFirstPlayerController())
-		{
-			const int32 HistoryFrame = ReplicatedDatas.ServerFrame - PlayerController->GetLocalToServerAsyncPhysicsTickOffset();
-			const FVehicleReplicatedDatas ServerDatas = ReplicatedDatas;
-
-			const int32 CurrentFrame = PlayerController->GetAsyncPhysicsTimestamp().LocalFrame;
-			PlayerController->ExecuteAsyncPhysicsCommand(PlayerController->GetAsyncPhysicsTimestamp(), this, [this, ServerDatas, HistoryFrame, bHasController, CurrentFrame]
-				{
-					if (HistoryFrame >= 0)
-					{
-						VehicleSimulationPT->InputsHistory->RecordInputs(HistoryFrame, ServerDatas.InputsDatas);
-						VehicleSimulationPT->StatesHistory->RecordStates(HistoryFrame, ServerDatas.StatesDatas);
-					}
-				}, false);
-		}
-	}
-}
-
-
 #undef LOCTEXT_NAMESPACE
 
 
diff --git a/Engine/Plugins/Experimental/ChaosVehiclesPlugin/Source/ChaosVehicles/Private/ChaosWheeledVehicleMovementComponent.cpp b/Engine/Plugins/Experimental/ChaosVehiclesPlugin/Source/ChaosVehicles/Private/ChaosWheeledVehicleMovementComponent.cpp
--- a/Engine/Plugins/Experimental/ChaosVehiclesPlugin/Source/ChaosVehicles/Private/ChaosWheeledVehicleMovementComponent.cpp
+++ b/Engine/Plugins/Experimental/ChaosVehiclesPlugin/Source/ChaosVehicles/Private/ChaosWheeledVehicleMovementComponent.cpp
@@ -171,7 +171,7 @@ bool UChaosWheeledVehicleSimulation::CanSimulate() const
 		&& PVehicle->Wheels.Num() == PVehicle->Suspension.Num());
 }
 
-void UChaosWheeledVehicleSimulation::TickVehicle(UWorld* WorldIn, float DeltaTime, const FChaosVehicleDefaultAsyncInput& InputData, FChaosVehicleAsyncOutput& OutputData, Chaos::FRigidBodyHandle_Internal* Handle)
+void UChaosWheeledVehicleSimulation::TickVehicle(UWorld* WorldIn, float DeltaTime, const FChaosVehicleAsyncInput& InputData, FChaosVehicleAsyncOutput& OutputData, Chaos::FRigidBodyHandle_Internal* Handle)
 {
 	SCOPE_CYCLE_COUNTER(STAT_ChaosVehicle_TickVehicle);
 
@@ -179,7 +179,7 @@ void UChaosWheeledVehicleSimulation::TickVehicle(UWorld* WorldIn, float DeltaTim
 }
 
 
-void UChaosWheeledVehicleSimulation::UpdateState(float DeltaTime, const FChaosVehicleDefaultAsyncInput& InputData, Chaos::FRigidBodyHandle_Internal* Handle)
+void UChaosWheeledVehicleSimulation::UpdateState(float DeltaTime, const FChaosVehicleAsyncInput& InputData, Chaos::FRigidBodyHandle_Internal* Handle)
 {
 	UChaosVehicleSimulation::UpdateState(DeltaTime, InputData, Handle);
 
@@ -233,7 +233,7 @@ void UChaosWheeledVehicleSimulation::UpdateState(float DeltaTime, const FChaosVe
 
 		if (!GWheeledVehicleDebugParams.DisableSuspensionForces && PVehicle->bSuspensionEnabled)
 		{
-			PerformSuspensionTraces(WheelState.Trace, InputData.TraceParams, InputData.TraceCollisionResponse, InputData.WheelTraceParams);
+			PerformSuspensionTraces(WheelState.Trace, InputData.PhysicsInputs.TraceParams, InputData.PhysicsInputs.TraceCollisionResponse, InputData.PhysicsInputs.WheelTraceParams);
 		}
 
 		//////////////////////////////////////////////////////////////////////////
@@ -257,7 +257,7 @@ void UChaosWheeledVehicleSimulation::UpdateState(float DeltaTime, const FChaosVe
 	}
 }
 
-void UChaosWheeledVehicleSimulation::UpdateSimulation(float DeltaTime, const FChaosVehicleDefaultAsyncInput& InputData, Chaos::FRigidBodyHandle_Internal* Handle)
+void UChaosWheeledVehicleSimulation::UpdateSimulation(float DeltaTime, const FChaosVehicleAsyncInput& InputData, Chaos::FRigidBodyHandle_Internal* Handle)
 {
 	SCOPE_CYCLE_COUNTER(STAT_ChaosVehicle_UpdateSimulation);
 	
@@ -278,13 +278,13 @@ void UChaosWheeledVehicleSimulation::UpdateSimulation(float DeltaTime, const FCh
 
 		if (!GWheeledVehicleDebugParams.DisableSuspensionForces && PVehicle->bSuspensionEnabled)
 		{
-			ApplySuspensionForces(DeltaTime, InputData.WheelTraceParams);
+			ApplySuspensionForces(DeltaTime, InputData.PhysicsInputs.WheelTraceParams);
 		}
 
 		///////////////////////////////////////////////////////////////////////
 		// Steering
 
-		ProcessSteering(InputData.ControlInputs);
+		ProcessSteering(InputData.PhysicsInputs.NetworkInputs.VehicleInputs);
 
 		///////////////////////////////////////////////////////////////////////
 		// Wheel Friction
@@ -530,15 +530,6 @@ void UChaosWheeledVehicleSimulation::ApplyWheelFrictionForces(float DeltaTime)
 			FMatrix Mat = FMatrix(GroundXVector, GroundYVector, GroundZVector, VehicleState.VehicleWorldTransform.GetLocation());
 			FVector FrictionForceVector = Mat.TransformVector(FrictionForceLocal);
 
-#if !(UE_BUILD_SHIPPING || UE_BUILD_TEST)
-			if (Chaos::FPhysicsSolverBase::IsNetworkPhysicsPredictionEnabled() && Chaos::FPhysicsSolverBase::CanDebugNetworkPhysicsPrediction())
-			{
-				UE_LOG(LogVehicle, Log, TEXT("Friction Force Global = %s | Transform = %s | Local = %s | Steering Angle = %f | Velocity = %s | Rotation Angle = %f | Friction Force = %s"),
-					*FrictionForceVector.ToString(), *Mat.ToString(), *FrictionForceLocal.ToString(), 
-					SteerAngleDegrees, *WheelState.LocalWheelVelocity[WheelIdx].ToString(), RotationAngle, *PWheel.GetForceFromFriction().ToString());
-			}
-#endif
-
 			check(PWheel.InContact());
 			if (PVehicle->bLegacyWheelFrictionPosition)
 			{
@@ -734,15 +725,6 @@ void UChaosWheeledVehicleSimulation::ProcessSteering(const FControlInputs& Contr
 
 			float SteeringAngle = ControlInputs.SteeringInput * SpeedScale;
 
-#if !(UE_BUILD_SHIPPING || UE_BUILD_TEST)
-			if (Chaos::FPhysicsSolverBase::IsNetworkPhysicsPredictionEnabled() && Chaos::FPhysicsSolverBase::CanDebugNetworkPhysicsPrediction())
-			{
-				UE_LOG(LogVehicle, Log, TEXT("Process Sterring Input = %f | Spped Scale = %f | Wheel Side = %f"),
-					ControlInputs.SteeringInput, SpeedScale, PVehicle->GetSuspension(WheelIdx).GetLocalRestingPosition().Y);
-			}
-#endif
-
-
 			if (FMath::Abs(GWheeledVehicleDebugParams.SteeringOverride) > 0.01f)
 			{
 				SteeringAngle = PWheel.MaxSteeringAngle * GWheeledVehicleDebugParams.SteeringOverride;
@@ -778,7 +760,7 @@ void UChaosWheeledVehicleSimulation::ApplyInput(const FControlInputs& ControlInp
 
 		if (ModifiedInputs.TransmissionType != PTransmission.Setup().TransmissionType)
 		{
-			PTransmission.AccessSetup().TransmissionType = ModifiedInputs.TransmissionType;
+			PTransmission.AccessSetup().TransmissionType = (Chaos::ETransmissionType)ModifiedInputs.TransmissionType;
 		}
 
 		if (ModifiedInputs.GearUpInput)
@@ -1717,7 +1699,7 @@ void UChaosWheeledVehicleMovementComponent::Update(float DeltaTime)
 		{
 			if (auto Handle = BodyInstance->ActorHandle)
 			{
-				FChaosVehicleDefaultAsyncInput* AsyncInput = static_cast<FChaosVehicleDefaultAsyncInput*>(CurAsyncInput);
+				FChaosVehicleAsyncInput* AsyncInput = static_cast<FChaosVehicleAsyncInput*>(CurAsyncInput);
 
 				TArray<AActor*> ActorsToIgnore;
 				ActorsToIgnore.Add(GetPawnOwner()); // ignore self in scene query
@@ -1726,14 +1708,14 @@ void UChaosWheeledVehicleMovementComponent::Update(float DeltaTime)
 				TraceParams.bReturnPhysicalMaterial = true;	// we need this to get the surface friction coefficient
 				TraceParams.AddIgnoredActors(ActorsToIgnore);
 				TraceParams.bTraceComplex = true;
-				AsyncInput->TraceParams = TraceParams;
-				AsyncInput->TraceCollisionResponse = WheelTraceCollisionResponses;
+				AsyncInput->PhysicsInputs.TraceParams = TraceParams;
+				AsyncInput->PhysicsInputs.TraceCollisionResponse = WheelTraceCollisionResponses;
 
-				AsyncInput->WheelTraceParams.SetNum(Wheels.Num());
+				AsyncInput->PhysicsInputs.WheelTraceParams.SetNum(Wheels.Num());
 				for (int I = 0; I < Wheels.Num(); I++)
 				{
-					AsyncInput->WheelTraceParams[I].SweepType = Wheels[I]->SweepType;
-					AsyncInput->WheelTraceParams[I].SweepShape = Wheels[I]->SweepShape;
+					AsyncInput->PhysicsInputs.WheelTraceParams[I].SweepType = Wheels[I]->SweepType;
+					AsyncInput->PhysicsInputs.WheelTraceParams[I].SweepShape = Wheels[I]->SweepShape;
 				}
 			}
 		}
diff --git a/Engine/Plugins/Experimental/ChaosVehiclesPlugin/Source/ChaosVehicles/Public/ChaosVehicleManagerAsyncCallback.h b/Engine/Plugins/Experimental/ChaosVehiclesPlugin/Source/ChaosVehicles/Public/ChaosVehicleManagerAsyncCallback.h
--- a/Engine/Plugins/Experimental/ChaosVehiclesPlugin/Source/ChaosVehicles/Public/ChaosVehicleManagerAsyncCallback.h
+++ b/Engine/Plugins/Experimental/ChaosVehiclesPlugin/Source/ChaosVehicles/Public/ChaosVehicleManagerAsyncCallback.h
@@ -7,7 +7,12 @@
 #include "Chaos/SimCallbackInput.h"
 #include "Chaos/SimCallbackObject.h"
 #include "Chaos/GeometryParticlesfwd.h"
+#include "TransmissionSystem.h"
 #include "PhysicsProxy/SingleParticlePhysicsProxyFwd.h"
+#include "Physics/NetworkPhysicsComponent.h"
+#include "ChaosVehicleWheel.h"
+
+#include "ChaosVehicleManagerAsyncCallback.generated.h"
 
 class UChaosVehicleMovementComponent;
 
@@ -19,6 +24,191 @@ enum EChaosAsyncVehicleDataType : int8
 	AsyncDefault,
 };
 
+/** Vehicle inputs from the player controller */
+USTRUCT()
+struct CHAOSVEHICLES_API FVehicleInputs
+{
+	GENERATED_USTRUCT_BODY()
+
+		FVehicleInputs()
+		: SteeringInput(0.f)
+		, ThrottleInput(0.f)
+		, BrakeInput(0.f)
+		, PitchInput(0.f)
+		, RollInput(0.f)
+		, YawInput(0.f)
+		, HandbrakeInput(0.f)
+	{}
+
+	// Steering output to physics system. Range -1...1
+	UPROPERTY()
+	float SteeringInput;
+
+	// Accelerator output to physics system. Range 0...1
+	UPROPERTY()
+	float ThrottleInput;
+
+	// Brake output to physics system. Range 0...1
+	UPROPERTY()
+	float BrakeInput;
+
+	// Body Pitch output to physics system. Range -1...1
+	UPROPERTY()
+	float PitchInput;
+
+	// Body Roll output to physics system. Range -1...1
+	UPROPERTY()
+	float RollInput;
+
+	// Body Yaw output to physics system. Range -1...1
+	UPROPERTY()
+	float YawInput;
+
+	// Handbrake output to physics system. Range 0...1
+	UPROPERTY()
+	float HandbrakeInput;
+};
+
+/** Control inputs holding the vehicle ones + persisten ones */
+USTRUCT()
+struct CHAOSVEHICLES_API FControlInputs : public FVehicleInputs
+{
+	GENERATED_USTRUCT_BODY()
+
+	FControlInputs()
+		: FVehicleInputs()
+		, ParkingEnabled(false)
+		, TransmissionType(Chaos::ETransmissionType::Automatic)
+		, GearUpInput(false)
+		, GearDownInput(false)
+	{}
+
+	/** Check if the parking is enabled  */
+	UPROPERTY()
+	bool ParkingEnabled;
+
+	/** Type of transmission (manual/auto) */
+	UPROPERTY()
+	uint8 TransmissionType;
+
+	/** Check if the user has risen the gear */
+	UPROPERTY()
+	bool GearUpInput;
+
+	/** Check if the user has descreased the gear */
+	UPROPERTY()
+	bool GearDownInput;
+};
+
+/** Vehicle Inputs datas that will be used in the inputs history to be applied while simulating */
+USTRUCT()
+struct CHAOSVEHICLES_API FNetworkVehicleInputs : public FNetworkPhysicsDatas
+{
+	GENERATED_USTRUCT_BODY()
+
+	/** List of incoming control inputs coming from the local client */
+	UPROPERTY()
+	FControlInputs VehicleInputs;
+
+	/** Transmission change time that could be set/changed from GT */
+	UPROPERTY()
+	float TransmissionChangeTime = 0.0;
+
+	/** Transmission current gear that could be set/changed from GT */
+	UPROPERTY()
+	int32 TransmissionCurrentGear = 0;
+
+	/** Transmission target gear that could be set/changed from GT */
+	UPROPERTY()
+	int32 TransmissionTargetGear = 0;
+
+	/**  Apply the datas onto the network physics component */
+	virtual void ApplyDatas(UActorComponent* NetworkComponent) const override;
+
+	/**  Build the datas from the network physics component */
+	virtual void BuildDatas(const UActorComponent* NetworkComponent) override;
+
+	/**  Serialize datas function that will be used to transfer the struct across the network */
+	bool NetSerialize(FArchive& Ar, class UPackageMap* Map, bool& bOutSuccess);
+
+	/** Interpolate the datas in between two inputs datas */
+	void InterpolateDatas(const FNetworkVehicleInputs& MinDatas, const FNetworkVehicleInputs& MaxDatas);
+};
+
+template<>
+struct TStructOpsTypeTraits<FNetworkVehicleInputs> : public TStructOpsTypeTraitsBase2<FNetworkVehicleInputs>
+{
+	enum
+	{
+		WithNetSerializer = true,
+	};
+};
+
+/** Vehicle states datas that will be used in the states history to rewind the simulation at some point inn time */
+USTRUCT()
+struct CHAOSVEHICLES_API FNetworkVehicleStates : public FNetworkPhysicsDatas
+{
+	GENERATED_USTRUCT_BODY()
+
+	/** Vehicle state last velocity */
+	UPROPERTY()
+	FVector StateLastVelocity = FVector(0,0,0);
+
+	/** Angular velocity for each wheels */
+	UPROPERTY()
+	TArray<float> WheelsOmega;
+
+	/** Angular position for each wheels */
+	UPROPERTY()
+	TArray<float> WheelsAngularPosition;
+
+	/** Suspension latest displacement to be used while simulating */
+	UPROPERTY()
+	TArray<float> SuspensionLastDisplacement;
+
+	/** Suspension latest spring length to be used while simulating  */
+	UPROPERTY()
+	TArray<float> SuspensionLastSpringLength;
+
+	/** Suspension averaged length for smoothing */
+	UPROPERTY()
+	TArray<float> SuspensionAveragedLength;
+
+	/** Suspension averaged count for smoothing */
+	UPROPERTY()
+	TArray<int32> SuspensionAveragedCount;
+
+	/** Suspension averaged number for smoothing */
+	UPROPERTY()
+	TArray<int32> SuspensionAveragedNum;
+
+	/** Engine angular velocity */
+	UPROPERTY()
+	float EngineOmega = 0.0;
+
+	/**  Apply the datas onto the network physics component */
+	virtual void ApplyDatas(UActorComponent* NetworkComponent) const override;
+
+	/**  Build the datas from the network physics component */
+	virtual void BuildDatas(const UActorComponent* NetworkComponent) override;
+
+	/**  Serialize datas function that will be used to transfer the struct across the network */
+	bool NetSerialize(FArchive& Ar, class UPackageMap* Map, bool& bOutSuccess);
+
+	/** Interpolate the datas in between two inputs datas */
+	void InterpolateDatas(const FNetworkVehicleStates& MinDatas, const FNetworkVehicleStates& MaxDatas);
+};
+
+template<>
+struct TStructOpsTypeTraits<FNetworkVehicleStates> : public TStructOpsTypeTraitsBase2<FNetworkVehicleStates>
+{
+	enum
+	{
+		WithNetSerializer = true,
+	};
+};
+
+/** Local physics wheels outputs not replicated across the network */
 struct CHAOSVEHICLES_API FWheelsOutput
 {
 	FWheelsOutput()
@@ -102,6 +292,37 @@ struct CHAOSVEHICLES_API FPhysicsVehicleOutput
 	float TransmissionTorque;
 };
 
+struct CHAOSVEHICLES_API FWheelTraceParams
+{
+	ESweepType SweepType;
+	ESweepShape SweepShape;
+};
+
+/**
+ * Per Vehicle input State from Game Thread to Physics Thread
+ */
+struct CHAOSVEHICLES_API FPhysicsVehicleInputs
+{
+	FPhysicsVehicleInputs()
+		: GravityZ(0.0f)
+		, TraceParams()
+		, TraceCollisionResponse()
+		, WheelTraceParams()
+	{
+	}
+	float GravityZ;
+	mutable FNetworkVehicleInputs NetworkInputs;
+	mutable FCollisionQueryParams TraceParams;
+	mutable FCollisionResponseContainer TraceCollisionResponse;
+	mutable TArray<FWheelTraceParams> WheelTraceParams;
+};
+
+struct CHAOSVEHICLES_API FPhysicsVehicleTraits
+{
+	using InputsType = FNetworkVehicleInputs;
+	using StatesType = FNetworkVehicleStates;
+};
+
 /**
  * Per Vehicle Input State from Game Thread to Physics Thread
  */                             
@@ -109,15 +330,16 @@ struct CHAOSVEHICLES_API FChaosVehicleAsyncInput
 {
 	const EChaosAsyncVehicleDataType Type;
 	UChaosVehicleMovementComponent* Vehicle;
-	
 	Chaos::FSingleParticlePhysicsProxy* Proxy;
 
+	FPhysicsVehicleInputs PhysicsInputs;
+
 	/** 
 	* Vehicle simulation running on the Physics Thread
 	*/
-	virtual TUniquePtr<struct FChaosVehicleAsyncOutput> Simulate(UWorld* World, const float DeltaSeconds, const float TotalSeconds, bool& bWakeOut) const = 0;
+	virtual TUniquePtr<struct FChaosVehicleAsyncOutput> Simulate(UWorld* World, const float DeltaSeconds, const float TotalSeconds, bool& bWakeOut) const;
 
-	virtual void ApplyDeferredForces(Chaos::FRigidBodyHandle_Internal* RigidHandle) const = 0;
+	virtual void ApplyDeferredForces(Chaos::FRigidBodyHandle_Internal* RigidHandle) const;
 
 	FChaosVehicleAsyncInput(EChaosAsyncVehicleDataType InType = EChaosAsyncVehicleDataType::AsyncInvalid)
 		: Type(InType)
@@ -146,7 +368,7 @@ struct FChaosVehicleManagerAsyncInput : public Chaos::FSimCallbackInput
 /**
  * Async Output Data
  */
-struct CHAOSVEHICLES_API FChaosVehicleAsyncOutput
+struct CHAOSVEHICLES_API FChaosVehicleAsyncOutput 
 {
 	const EChaosAsyncVehicleDataType Type;
 	bool bValid;	// indicates no work was done
@@ -160,7 +382,6 @@ struct CHAOSVEHICLES_API FChaosVehicleAsyncOutput
 	virtual ~FChaosVehicleAsyncOutput() = default;
 };
 
-
 /**
  * Async Output for all of the vehicles handled by this Vehicle Manager
  */
@@ -178,11 +399,11 @@ struct FChaosVehicleManagerAsyncOutput : public Chaos::FSimCallbackOutput
 /**
  * Async callback from the Physics Engine where we can perform our vehicle simulation
  */
-class CHAOSVEHICLES_API FChaosVehicleManagerAsyncCallback : public Chaos::TSimCallbackObject<FChaosVehicleManagerAsyncInput, FChaosVehicleManagerAsyncOutput, Chaos::ESimCallbackOptions::Presimulate | Chaos::ESimCallbackOptions::ContactModification>
+class CHAOSVEHICLES_API FChaosVehicleManagerAsyncCallback : public Chaos::TSimCallbackObject<FChaosVehicleManagerAsyncInput, FChaosVehicleManagerAsyncOutput, Chaos::ESimCallbackOptions::Presimulate | Chaos::ESimCallbackOptions::Rewind>
 {
 public:
 	virtual FName GetFNameForStatId() const override;
 private:
+	virtual void ProcessInputs_Internal(int32 PhysicsStep) override;
 	virtual void OnPreSimulate_Internal() override;
-	virtual void OnContactModification_Internal(Chaos::FCollisionContactModifier& Modifications) override;
 };
diff --git a/Engine/Plugins/Experimental/ChaosVehiclesPlugin/Source/ChaosVehicles/Public/ChaosVehicleMovementComponent.h b/Engine/Plugins/Experimental/ChaosVehiclesPlugin/Source/ChaosVehicles/Public/ChaosVehicleMovementComponent.h
--- a/Engine/Plugins/Experimental/ChaosVehiclesPlugin/Source/ChaosVehicles/Public/ChaosVehicleMovementComponent.h
+++ b/Engine/Plugins/Experimental/ChaosVehiclesPlugin/Source/ChaosVehicles/Public/ChaosVehicleMovementComponent.h
@@ -12,6 +12,7 @@
 #include "ChaosVehicleWheel.h"
 #include "AerofoilSystem.h"
 #include "ThrustSystem.h"
+#include "TransmissionSystem.h"
 #include "PhysicsProxy/SingleParticlePhysicsProxyFwd.h"
 #include "SnapshotData.h"
 #include "DeferredForces.h"
@@ -30,218 +31,6 @@ class UCanvas;
 struct FChaosVehicleAsyncInput;
 struct FChaosVehicleManagerAsyncOutput;
 
-/** Vehicle inputs from the player controller */
-USTRUCT()
-struct CHAOSVEHICLES_API FVehicleInputs
-{
-	GENERATED_USTRUCT_BODY()
-
-		FVehicleInputs()
-		: SteeringInput(0.f)
-		, ThrottleInput(0.f)
-		, BrakeInput(0.f)
-		, PitchInput(0.f)
-		, RollInput(0.f)
-		, YawInput(0.f)
-		, HandbrakeInput(0.f)
-	{}
-
-	/** Copy the parent class properties from another instance */
-	void CopyVehicleInputs(const FVehicleInputs& VehicleInputs)
-	{
-		SteeringInput = VehicleInputs.SteeringInput;
-		ThrottleInput = VehicleInputs.ThrottleInput;
-		BrakeInput = VehicleInputs.BrakeInput;
-		PitchInput = VehicleInputs.PitchInput;
-		RollInput = VehicleInputs.RollInput;
-		YawInput = VehicleInputs.YawInput;
-		HandbrakeInput = VehicleInputs.HandbrakeInput;
-	}
-
-	// Steering output to physics system. Range -1...1
-	UPROPERTY()
-	float SteeringInput;
-
-	// Accelerator output to physics system. Range 0...1
-	UPROPERTY()
-	float ThrottleInput;
-
-	// Brake output to physics system. Range 0...1
-	UPROPERTY()
-	float BrakeInput;
-
-	// Body Pitch output to physics system. Range -1...1
-	UPROPERTY()
-	float PitchInput;
-
-	// Body Roll output to physics system. Range -1...1
-	UPROPERTY()
-	float RollInput;
-
-	// Body Yaw output to physics system. Range -1...1
-	UPROPERTY()
-	float YawInput;
-
-	// Handbrake output to physics system. Range 0...1
-	UPROPERTY()
-	float HandbrakeInput;
-};
-
-USTRUCT()
-struct CHAOSVEHICLES_API FControlInputs : public FVehicleInputs
-{
-	GENERATED_USTRUCT_BODY()
-
-		FControlInputs()
-		: FVehicleInputs()
-		, ParkingEnabled(false)
-		, TransmissionType(Chaos::ETransmissionType::Automatic)
-		, GearUpInput(false)
-		, GearDownInput(false)
-	{}
-
-	bool ParkingEnabled;
-
-	Chaos::ETransmissionType TransmissionType;
-
-	bool GearUpInput;
-	bool GearDownInput;
-};
-
-USTRUCT()
-struct CHAOSVEHICLES_API FVehicleHistoryDatas
-{
-	GENERATED_USTRUCT_BODY()
-
-	int32 DatasFrame = INDEX_NONE;
-
-	int32 SmoothFrame = INDEX_NONE;
-};
-
-/** Vehicle Inputs datas that will be used in the Inputs history to be applied while simulating */
-USTRUCT()
-struct CHAOSVEHICLES_API FVehicleInputsDatas : public FVehicleHistoryDatas
-{
-	GENERATED_USTRUCT_BODY()
-
-	FVehicleInputsDatas();
-
-	/** List of incoming control inputs coming from the local client */
-	UPROPERTY()
-	FVehicleInputs VehicleInputs;
-
-	/** Transmission change time that could be set/changed from GT */
-	UPROPERTY()
-	float TransmissionChangeTime;
-
-	/** Transmission current gear that could be set/changed from GT */
-	UPROPERTY()
-	int32 TransmissionCurrentGear;
-
-	/** Transmission target gear that could be set/changed from GT */
-	UPROPERTY()
-	int32 TransmissionTargetGear;
-};
-
-/** Vehicle states datas that will be used in the states history to rewind the simulation at some point inn time */
-USTRUCT()
-struct CHAOSVEHICLES_API FVehicleStatesDatas : public FVehicleHistoryDatas
-{
-	GENERATED_USTRUCT_BODY()
-
-	FVehicleStatesDatas();
-
-	/** Vehicle state last velocity */
-	UPROPERTY()
-	FVector StateLastVelocity;
-
-	/** Angular velocity for each wheels */
-	UPROPERTY()
-	TArray<float> WheelsOmega;
-
-	/** Angular position for each wheels */
-	UPROPERTY()
-	TArray<float> WheelsAngularPosition;
-
-	/** Suspension latest displacement to be used while simulating */
-	UPROPERTY()
-	TArray<float> SuspensionLastDisplacement;
-
-	/** Suspension latest spring length to be used while simulating  */
-	UPROPERTY()
-	TArray<float> SuspensionLastSpringLength;
-
-	/** Suspension averaged length for smoothing */
-	UPROPERTY()
-	TArray<float> SuspensionAveragedLength;
-
-	/** Suspension averaged count for smoothing */
-	UPROPERTY()
-	TArray<int32> SuspensionAveragedCount;
-
-	/** Suspension averaged number for smoothing */
-	UPROPERTY()
-	TArray<int32> SuspensionAveragedNum;
-
-	/** Engine angular velocity */
-	UPROPERTY()
-	float EngineOmega;
-};
-
-/** Vehicle states history that will be registered in the rewind datas */
-struct CHAOSVEHICLES_API FVehicleStatesHistory : public Chaos::FStatesRewindHistory<FVehicleStatesDatas>
-{
-	using Super = Chaos::FStatesRewindHistory<FVehicleStatesDatas>;
-	using DatasType = FVehicleStatesDatas;
-
-	FVehicleStatesHistory(const int32 FrameCount, const bool bIsHistoryLocal, class UChaosVehicleSimulation* ChaosSimulation);
-
-	/** Rewind the vehicle state at some point in time */
-	virtual bool RewindStates(const int32 RewindFrame, const bool bResetSolver) override;
-
-	/** Record the vehicle state from the replicated datas */
-	bool RecordStates(const int32 RecordFrame, const DatasType& ReplicatedHistory);
-
-	/** Replicate the states from the vehicle simulation*/
-	void ReplicateStates(DatasType& ReplicatedDatas);
-
-	virtual ~FVehicleStatesHistory() override;
-
-private:
-
-	class UChaosVehicleSimulation* VehicleSimulation = nullptr;
-};
-
-/** Vehicle Inputs history that will be registered in the rewind datas */
-struct CHAOSVEHICLES_API FVehicleInputsHistory : public Chaos::FInputsRewindHistory<FVehicleInputsDatas>
-{
-	using Super = Chaos::FInputsRewindHistory<FVehicleInputsDatas>;
-	using DatasType = FVehicleInputsDatas;
-
-	FVehicleInputsHistory(const int32 FrameCount, const bool bIsHistoryLocal, class UChaosVehicleSimulation* ChaosSimulation);
-
-	/** Apply the Inputs at some point in time */
-	virtual bool ApplyInputs(const int32 ApplyFrame, const bool bResetSolver) override;
-
-	/** Record the vehicle Inputs from the replicated datas */
-	bool RecordInputs(const int32 RecordFrame, const DatasType& ReplicatedHistory);
-
-	/** Replicate the Inputs from the vehicle simulation + external inputs */
-	void ReplicateInputs(
-		const float SteeringInput, const float ThrottleInput,
-		const float BrakeInput, const float HandbrakeInput, const float RollInput,
-		const float PitchInput, const float YawInput, const int32 CurrentGear, DatasType& ReplicatedDatas);
-
-	/** Smooth Inputs inputs (WIP) */
-	void SmoothInputs(const int32 ApplyFrame, const int32 ValidFrame, const bool bApplySmoothing);
-
-	virtual ~FVehicleInputsHistory() override;
-
-private:
-
-	class UChaosVehicleSimulation* VehicleSimulation = nullptr;
-};
-
 struct CHAOSVEHICLES_API FVehicleDebugParams
 {
 	bool ShowCOM = false;
@@ -291,50 +80,6 @@ struct CHAOSVEHICLES_API FVehicleReplicatedState : public FVehicleInputs
 	float ThrottleDown;
 };
 
-/** Group of datas that needs to be replicated across all clients for rewind/resim */
-USTRUCT()
-struct CHAOSVEHICLES_API FVehicleReplicatedDatas
-{
-	GENERATED_USTRUCT_BODY()
-
-	FVehicleReplicatedDatas()
-	{}
-
-	/** Server frame at which the datas have been recorded */
-	UPROPERTY()
-	int32 ServerFrame = INDEX_NONE;
-
-	/** Group of inputs to replicate */
-	UPROPERTY()
-	FVehicleStatesDatas StatesDatas;
-
-	/** Group of inputs to replicate */
-	UPROPERTY()
-	FVehicleInputsDatas InputsDatas;
-};
-
-struct CHAOSVEHICLES_API FWheelTraceParams
-{
-	ESweepType SweepType;
-	ESweepShape SweepShape;
-};
-
-struct CHAOSVEHICLES_API FChaosVehicleDefaultAsyncInput : public FChaosVehicleAsyncInput
-{
-	float GravityZ;
-	mutable FControlInputs ControlInputs;
-	mutable FCollisionQueryParams TraceParams;
-	mutable FCollisionResponseContainer TraceCollisionResponse;
-	mutable TArray<FWheelTraceParams> WheelTraceParams;
-
-	FChaosVehicleDefaultAsyncInput();
-
-	virtual TUniquePtr<FChaosVehicleAsyncOutput> Simulate(UWorld* World, const float DeltaSeconds, const float TotalSeconds, bool& bWakeOut) const override;
-
-	virtual void ApplyDeferredForces(Chaos::FRigidBodyHandle_Internal* RigidHandle) const override;
-};
-
-
 USTRUCT()
 struct CHAOSVEHICLES_API FVehicleTorqueControlConfig
 {
@@ -886,19 +631,19 @@ public:
 
 	virtual void UpdateConstraintHandles(TArray<FPhysicsConstraintHandle>& ConstraintHandlesIn) {}
 
-	virtual void TickVehicle(UWorld* WorldIn, float DeltaTime, const FChaosVehicleDefaultAsyncInput& InputData, FChaosVehicleAsyncOutput& OutputData, Chaos::FRigidBodyHandle_Internal* Handle);
+	virtual void TickVehicle(UWorld* WorldIn, float DeltaTime, const FChaosVehicleAsyncInput& InputData, FChaosVehicleAsyncOutput& OutputData, Chaos::FRigidBodyHandle_Internal* Handle);
 
 	/** Update the async inputs and the history ones*/
-	void SyncHistoryInputs(const FChaosVehicleDefaultAsyncInput& InputData, Chaos::FRigidBodyHandle_Internal* Handle);
+	void SyncHistoryInputs(const FChaosVehicleAsyncInput& InputData, Chaos::FRigidBodyHandle_Internal* Handle);
 
 	/** Apply the defered forces onto the vehicles */
 	virtual void ApplyDeferredForces(Chaos::FRigidBodyHandle_Internal* Handle);
 
 	/** Update the vehicle state */
-	virtual void UpdateState(float DeltaTime, const FChaosVehicleDefaultAsyncInput& InputData, Chaos::FRigidBodyHandle_Internal* Handle);
+	virtual void UpdateState(float DeltaTime, const FChaosVehicleAsyncInput& InputData, Chaos::FRigidBodyHandle_Internal* Handle);
 
 	/** Advance the vehicle simulation */
-	virtual void UpdateSimulation(float DeltaTime, const FChaosVehicleDefaultAsyncInput& InputData, Chaos::FRigidBodyHandle_Internal* Handle);
+	virtual void UpdateSimulation(float DeltaTime, const FChaosVehicleAsyncInput& InputData, Chaos::FRigidBodyHandle_Internal* Handle);
 
 	/** Fill the vehicle output state */
 	virtual void FillOutputState(FChaosVehicleAsyncOutput& Output);
@@ -919,7 +664,7 @@ public:
 	virtual void ApplyThrustForces(float DeltaTime);
 
 	/** Apply direct control over vehicle body rotation */
-	virtual void ApplyTorqueControl(float DeltaTime, const FChaosVehicleDefaultAsyncInput& InputData);
+	virtual void ApplyTorqueControl(float DeltaTime, const FChaosVehicleAsyncInput& InputData);
 
 	/** Add a force to this vehicle */
 	void AddForce(const FVector& Force, bool bAllowSubstepping = true, bool bAccelChange = false);
@@ -957,13 +702,7 @@ public:
 	FDeferredForces DeferredForces;
 
 	/** Current control inputs that is being used on the PT */
-	FVehicleInputs VehicleInputs;
-
-	/** States history uses to rewind simulation */
-	TSharedPtr<FVehicleStatesHistory> StatesHistory;
-
-	/** Inputs history used during simulation*/
-	TSharedPtr<FVehicleInputsHistory> InputsHistory;
+	FControlInputs VehicleInputs;
 
 	FVehicleInputRateConfig ThrottleInputRate;
 	FVehicleInputRateConfig SteeringInputRate;
@@ -981,8 +720,14 @@ public:
 UCLASS(Abstract, hidecategories = (PlanarMovement, "Components|Movement|Planar", Activation, "Components|Activation"))
 class CHAOSVEHICLES_API UChaosVehicleMovementComponent : public UPawnMovementComponent
 {
-	friend struct FChaosVehicleDefaultAsyncInput;
+	friend struct FChaosVehicleAsyncInput;
+	friend struct FChaosVehicleAsyncOutput;
+
+	friend struct FNetworkVehicleInputs;
+	friend struct FNetworkVehicleStates;
+
 	friend class FChaosVehicleManager;
+	friend class FChaosVehicleManagerAsyncCallback;
 
 	GENERATED_UCLASS_BODY()
 
@@ -1095,6 +840,9 @@ protected:
 
 	Chaos::ETransmissionType TransmissionType;
 
+	UPROPERTY()
+	TObjectPtr<UNetworkPhysicsComponent> NetworkPhysicsComponent = nullptr;
+
 public:
 
 	/** UObject interface */
@@ -1134,7 +882,6 @@ public:
 	/** Stops movement immediately (zeroes velocity, usually zeros acceleration for components with acceleration). */
 	virtual void StopMovementImmediately() override;
 
-
 	/** Set the user input for the vehicle throttle [range 0 to 1] */
 	UFUNCTION(BlueprintCallable, Category = "Game|Components|ChaosVehicleMovement")
 	void SetThrottleInput(float Throttle);
@@ -1302,19 +1049,12 @@ public:
 	};
 	TArray<FAsyncOutputWrapper> OutputsWaitingOn;
 
-	UFUNCTION()
-	virtual void OnRep_SetVehicleHistory();
-
 protected:
 
 	// replicated state of vehicle 
 	UPROPERTY(Transient, Replicated)
 	FVehicleReplicatedState ReplicatedState;
 
-	// replicated state of vehicle 
-	UPROPERTY(Transient, ReplicatedUsing = OnRep_SetVehicleHistory)
-	FVehicleReplicatedDatas ReplicatedDatas;
-
 	// accumulator for RB replication errors 
 	float AngErrorAccumulator;
 
@@ -1474,18 +1214,12 @@ protected:
 
 	// Setup
 
-	/** New method of applying control inputs at a specific time on the server - for network prediction */
-	UFUNCTION(Server, Reliable)
-	void ServerApplyControls(float InSteeringInput, float InThrottleInput, float InBrakeInput, float InHandbrakeInput,
-			int32 InCurrentGear, float InRollInput, float InPitchInput, float InYawInput, const FAsyncPhysicsTimestamp& AsyncPhysicsTimestamp);
-
-	/** Apply Controls locally */
-	void ApplyControls_Imp(float InSteeringInput, float InThrottleInput, float InBrakeInput, float InHandbrakeInput,
-		int32 InCurrentGear, float InRollInput, float InPitchInput, float InYawInput, const FAsyncPhysicsTimestamp& AsyncPhysicsTimestamp, const bool bEnableResim = true);
-
 	/** Get our controller */
 	virtual AController* GetController() const override;
 
+	/** Retrieve the player controller of the vehicle component */
+	APlayerController* GetPlayerController() const;
+
 	/** Get the mesh this vehicle is tied to */
 	class UMeshComponent* GetMesh() const;
 
diff --git a/Engine/Plugins/Experimental/ChaosVehiclesPlugin/Source/ChaosVehicles/Public/ChaosWheeledVehicleMovementComponent.h b/Engine/Plugins/Experimental/ChaosVehiclesPlugin/Source/ChaosVehicles/Public/ChaosWheeledVehicleMovementComponent.h
--- a/Engine/Plugins/Experimental/ChaosVehiclesPlugin/Source/ChaosVehicles/Public/ChaosWheeledVehicleMovementComponent.h
+++ b/Engine/Plugins/Experimental/ChaosVehiclesPlugin/Source/ChaosVehicles/Public/ChaosWheeledVehicleMovementComponent.h
@@ -585,13 +585,13 @@ public:
 
 	virtual void UpdateConstraintHandles(TArray<FPhysicsConstraintHandle>& ConstraintHandlesIn) override;
 
-	virtual void TickVehicle(UWorld* WorldIn, float DeltaTime, const FChaosVehicleDefaultAsyncInput& InputData, FChaosVehicleAsyncOutput& OutputData, Chaos::FRigidBodyHandle_Internal* Handle) override;
+	virtual void TickVehicle(UWorld* WorldIn, float DeltaTime, const FChaosVehicleAsyncInput& InputData, FChaosVehicleAsyncOutput& OutputData, Chaos::FRigidBodyHandle_Internal* Handle) override;
 
 	/** Advance the vehicle simulation */
-	virtual void UpdateSimulation(float DeltaTime, const FChaosVehicleDefaultAsyncInput& InputData, Chaos::FRigidBodyHandle_Internal* Handle) override;
+	virtual void UpdateSimulation(float DeltaTime, const FChaosVehicleAsyncInput& InputData, Chaos::FRigidBodyHandle_Internal* Handle) override;
 
 	/** Update the vehicle state */
-	virtual void UpdateState(float DeltaTime, const FChaosVehicleDefaultAsyncInput& InputData, Chaos::FRigidBodyHandle_Internal* Handle) override;
+	virtual void UpdateState(float DeltaTime, const FChaosVehicleAsyncInput& InputData, Chaos::FRigidBodyHandle_Internal* Handle) override;
 
 	virtual void FillOutputState(FChaosVehicleAsyncOutput& Output) override;
 
diff --git a/Engine/Source/Runtime/Engine/Classes/GameFramework/PlayerController.h b/Engine/Source/Runtime/Engine/Classes/GameFramework/PlayerController.h
--- a/Engine/Source/Runtime/Engine/Classes/GameFramework/PlayerController.h
+++ b/Engine/Source/Runtime/Engine/Classes/GameFramework/PlayerController.h
@@ -386,22 +386,25 @@ protected:
 	virtual void BeginReplication() override;
 #endif // UE_WITH_IRIS
 	/** The type of async physics data object to use*/
-	UPROPERTY(EditDefaultsOnly, Category=PlayerController)
-	TSubclassOf<UAsyncPhysicsData> AsyncPhysicsDataClass;
+	UPROPERTY(meta = (DeprecatedProperty, DeprecationMessage = "AsyncPhysicsDataClass is deprecated. See UInputSettings::bEnableLegacyInputScales to enable legacy behavior"))
+	TSubclassOf<UAsyncPhysicsData> AsyncPhysicsDataClass_DEPRECATED;
 
 	/** Get the async physics data to write to. This data will make its way to the async physics tick on client and server. Should not be used during async tick */
+	UE_DEPRECATED(5.3, "GetAsyncPhysicsDataToConsume is deprecated, please see the new C++ NetworkPhysicsComponent")
 	UFUNCTION(BlueprintPure, Category = PlayerController)
 	UAsyncPhysicsData* GetAsyncPhysicsDataToWrite() const;
 
 	/** Get the async physics data to execute logic off of. This data should not be modified and will NOT make its way back. Must be used during async tick */
+	UE_DEPRECATED(5.3, "GetAsyncPhysicsDataToConsume is deprecated, please see the new C++ NetworkPhysicsComponent")
 	UFUNCTION(BlueprintPure, Category = PlayerController)
 	const UAsyncPhysicsData* GetAsyncPhysicsDataToConsume() const;
 
 private:
 
-	UPROPERTY(ReplicatedUsing=OnRep_AsyncPhysicsDataComponent)
-	TObjectPtr<UAsyncPhysicsInputComponent> AsyncPhysicsDataComponent;
+	UPROPERTY(ReplicatedUsing=OnRep_AsyncPhysicsDataComponent, meta = (DeprecatedProperty, DeprecationMessage = "AsyncPhysicsDataComponent is deprecated. please see the new C++ NetworkPhysicsComponent"))
+	TObjectPtr<UAsyncPhysicsInputComponent> AsyncPhysicsDataComponent_DEPRECARED;
 
+	UE_DEPRECATED(5.3, "OnRep_AsyncPhysicsDataComponent is deprecated, please see the new C++ NetworkPhysicsComponent")
 	UFUNCTION()
 	void OnRep_AsyncPhysicsDataComponent();
 
@@ -2343,6 +2346,10 @@ private:
 	/** The estimated offset between the local async physics tick frame number and the server's
 	*	This is used to synchronize events that happen in the async physics tick */
 	int32 LocalToServerAsyncPhysicsTickOffset;
+
+	/** The estimated offset between the server async physics tick frame number and the local's
+	*	This is used to synchronize events that happen in the async physics tick */
+	int32 ServerToLocalAsyncPhysicsTickOffset = INDEX_NONE;
 	
 	/** The latest server step we've received an offset correction for. This allows us to ignore out of order corrections that arrive late */
 	int32 ClientLatestCorrectedOffsetServerStep = INDEX_NONE;
@@ -2391,6 +2398,12 @@ public:
 	/** Returns the current estimated offset between the local async physics step and the server. This is useful for dealing with low level synchronization.
 		In general it's recommended to use GetAsyncPhysicsTimestamp which accounts for the offset automatically*/
 	int32 GetLocalToServerAsyncPhysicsTickOffset() const { return LocalToServerAsyncPhysicsTickOffset; }
+
+	/** Returns the current estimated offset between the server async physics step and the local one. */
+	int32 GetServerToLocalAsyncPhysicsTickOffset() const { return (ServerToLocalAsyncPhysicsTickOffset != INDEX_NONE) ? ServerToLocalAsyncPhysicsTickOffset : LocalToServerAsyncPhysicsTickOffset; }
+
+	/** Set the offset between the server async physics step and the local one.*/
+	void SetServerToLocalAsyncPhysicsTickOffset( const int32 AsyncPhysicsTickOffset) { ServerToLocalAsyncPhysicsTickOffset = AsyncPhysicsTickOffset; }
 	
 };
 
diff --git a/Engine/Source/Runtime/Engine/Private/PhysicsEngine/AsyncPhysicsInputComponent.cpp b/Engine/Source/Runtime/Engine/Private/PhysicsEngine/AsyncPhysicsInputComponent.cpp
--- a/Engine/Source/Runtime/Engine/Private/PhysicsEngine/AsyncPhysicsInputComponent.cpp
+++ b/Engine/Source/Runtime/Engine/Private/PhysicsEngine/AsyncPhysicsInputComponent.cpp
@@ -16,70 +16,10 @@
 
 namespace InputCmdCVars
 {
-	static int32 ForceFault = 0;
-	static FAutoConsoleVariableRef CVarForceFault(TEXT("p.net.ForceFault"), ForceFault, TEXT("Forces server side input fault"));
-
 	static int32 ForceInputDrop = 0;
 	static FAutoConsoleVariableRef CVarForceInputDrop(TEXT("p.net.ForceInputDrop"), ForceInputDrop, TEXT("Forces client to drop inputs. Useful for simulating desync"));
-
-	static int32 MaxBufferedCmds = 16;
-	static FAutoConsoleVariableRef CVarMaxBufferedCmds(TEXT("p.net.MaxBufferedCmds"), MaxBufferedCmds, TEXT("MaxNumber of buffered server side commands"));
-
-	static int32 TimeDilationEnabled = 1;
-	static FAutoConsoleVariableRef CVarTimeDilationEnabled(TEXT("p.net.TimeDilationEnabled"), TimeDilationEnabled, TEXT("Enable clientside TimeDilation"));
-
-	static float MaxTargetNumBufferedCmds = 5.0;
-	static FAutoConsoleVariableRef CVarMaxTargetNumBufferedCmds(TEXT("p.net.MaxTargetNumBufferedCmds"), MaxTargetNumBufferedCmds, TEXT("Maximum number of buffered inputs the server will target per client."));
-
-	static float MaxTimeDilationMag = 0.01f;
-	static FAutoConsoleVariableRef CVarMaxTimeDilationMag(TEXT("p.net.MaxTimeDilationMag"), MaxTimeDilationMag, TEXT("Maximum time dilation that client will use to slow down / catch up with server"));
-
-	static float TimeDilationAlpha = 0.1f;
-	static FAutoConsoleVariableRef CVarTimeDilationAlpha(TEXT("p.net.TimeDilationAlpha"), TimeDilationAlpha, TEXT("Lerp strength for sliding client time dilation"));
-
-	static float TargetNumBufferedCmdsDeltaOnFault = 1.0f;
-	static FAutoConsoleVariableRef CVarTargetNumBufferedCmdsDeltaOnFault(TEXT("p.net.TargetNumBufferedCmdsDeltaOnFault"), TargetNumBufferedCmdsDeltaOnFault, TEXT("How much to increase TargetNumBufferedCmds when an input fault occurs"));
-
-	static float TargetNumBufferedCmds = 1.9f;
-	static FAutoConsoleVariableRef CVarTargetNumBufferedCmds(TEXT("p.net.TargetNumBufferedCmds"), TargetNumBufferedCmds, TEXT("How much to increase TargetNumBufferedCmds when an input fault occurs"));
-
-	static float TargetNumBufferedCmdsAlpha = 0.005f;
-	static FAutoConsoleVariableRef CVarTargetNumBufferedCmdsAlpha(TEXT("p.net.TargetNumBufferedCmdsAlpha"), TargetNumBufferedCmdsAlpha, TEXT("Lerp strength for TargetNumBufferedCmds"));
-
-	static int32 LerpTargetNumBufferedCmdsAggresively = 0;
-	static FAutoConsoleVariableRef CVarLerpTargetNumBufferedCmdsAggresively(TEXT("p.net.LerpTargetNumBufferedCmdsAggresively"), LerpTargetNumBufferedCmdsAggresively, TEXT("Aggresively lerp towards TargetNumBufferedCmds. Reduces server side buffering but can cause more artifacts."));
 }
 
-
-// --------------------------------------------------------------------------------------------------------------------------------------------------
-//	Client InputCmd Stream stuff
-// --------------------------------------------------------------------------------------------------------------------------------------------------
-
-
-int8 QuantizeTimeDilation(float F)
-{
-	if (F == 1.f)
-	{
-		return 0;
-	}
-
-	float Normalized = FMath::Clamp<float>((F - 1.f) / InputCmdCVars::MaxTimeDilationMag, -1.f, 1.f);
-	return (int8)(Normalized * 128.f);
-}
-
-float DeQuantizeTimeDilation(int8 i)
-{
-	if (i == 0)
-	{
-		return 1.f;
-	}
-
-	float Normalized = (float)i / 128.f;
-	float Uncompressed = 1.f + (Normalized * InputCmdCVars::MaxTimeDilationMag);
-	return Uncompressed;
-}
-
-
 DECLARE_MULTICAST_DELEGATE_ThreeParams(FOnDispatchPhysicsTick, int32 PhysicsStep, int32 NumSteps, int32 ServerFrame);
 
 struct FAsyncPhysicsInputRewindCallback : public Chaos::IRewindCallback
@@ -116,7 +56,7 @@ struct FAsyncPhysicsInputRewindCallback : public Chaos::IRewindCallback
 		DispatchPhysicsTick.Broadcast(PhysicsStep, NumSteps, ServerFrame);
 	}
 
-	virtual int32 TriggerRewindIfNeeded_Internal(int32 LatestStepCompleted)
+	virtual int32 TriggerRewindIfNeeded_Internal(int32 LatestStepCompleted) override
 	{
 		int32 CachedClientFrame = INDEX_NONE;
 		if (World->GetNetMode() == NM_Client && World->GetPhysicsScene())
@@ -142,14 +82,6 @@ struct FAsyncPhysicsInputRewindCallback : public Chaos::IRewindCallback
 			int32 LocalOffset = 0;
 			if (APlayerController* PC = World->GetFirstPlayerController())
 			{
-				// ------------------------------------------------
-				// Send RPC to server telling them what (client/local) physics step we are running
-				//	* Note that SendData is empty because of the existing API, should change this
-				// ------------------------------------------------	
-
-				TArray<uint8> SendData;
-				PC->PushClientInput(PhysicsStep, SendData);
-
 				// -----------------------------------------------------------------
 				// Calculate latest frame offset to map server frame to local frame
 				// -----------------------------------------------------------------
@@ -165,150 +97,6 @@ struct FAsyncPhysicsInputRewindCallback : public Chaos::IRewindCallback
 					const int32 FrameOffset = ClientFrameInfo.GetLocalFrameOffset();
 					CachedServerFrame = PhysicsStep - FrameOffset; // Local = Server + Offset
 				}
-
-				// -----------------------------------------------------------------
-				// Apply local TIme Dilation based on server's recommendation.
-				// This speeds up or slows down our consumption of real time (by like < 1%)
-				// Ultimately this causes us to send InputCmds at a lower or higher rate in
-				// order to keep server side buffer at optimal capacity. 
-				// Optimal capacity = as small as possible without ever "missing" a frame (e.g, minimal buffer yet always a new fresh cmd to consume server side)
-				// -----------------------------------------------------------------
-				const float RealTimeDilation = DeQuantizeTimeDilation(ClientFrameInfo.QuantizedTimeDilation);
-
-				if (InputCmdCVars::TimeDilationEnabled > 0)
-				{
-					FPhysScene* PhysScene = World->GetPhysicsScene();
-					PhysScene->SetNetworkDeltaTimeScale(RealTimeDilation);
-				}
-			}
-		}
-		else
-		{
-			// -----------------------------------------------
-			// Server: "consume" an InputCmd from each Player Controller
-			// All this means in this context is updating FServerFrameInfo:: LastProcessedInputFrame, LastLocalFrame
-			// (E.g, telling each client what "Input" of theirs we were processing and our local physics frame number.
-			// In cases where the buffer has a fault, we calculate a suggested time dilation to temporarily make client speed up 
-			// or slow down their input cmd production.
-			// -----------------------------------------------
-
-			const bool bForceFault = InputCmdCVars::ForceFault > 0;
-			InputCmdCVars::ForceFault = FMath::Max(0, InputCmdCVars::ForceFault - 1);
-
-			for (FConstPlayerControllerIterator Iterator = World->GetPlayerControllerIterator(); Iterator; ++Iterator)
-			{
-				if (APlayerController* PC = Iterator->Get())
-				{
-					APlayerController::FServerFrameInfo& FrameInfo = PC->GetServerFrameInfo();
-					APlayerController::FInputCmdBuffer& InputBuffer = PC->GetInputBuffer();
-
-					{
-						const int32 NumBufferedInputCmds = bForceFault ? 0 : (InputBuffer.HeadFrame() - FrameInfo.LastProcessedInputFrame);
-
-						// Check Overflow
-						if (NumBufferedInputCmds > InputCmdCVars::MaxBufferedCmds)
-						{
-							UE_LOG(LogPhysics, Warning, TEXT("[Remote.Input] overflow %d %d -> %d"), InputBuffer.HeadFrame(), FrameInfo.LastProcessedInputFrame, NumBufferedInputCmds);
-							FrameInfo.LastProcessedInputFrame = InputBuffer.HeadFrame() - InputCmdCVars::MaxBufferedCmds + 1;
-						}
-
-						// Check fault - we are waiting for Cmds to reach TargetNumBufferedCmds before continuing
-						if (FrameInfo.bFault)
-						{
-							if (NumBufferedInputCmds < (int32)FrameInfo.TargetNumBufferedCmds)
-							{
-								// Skip this because it is in fault. We will use the prev input for this frame.
-								UE_CLOG(FrameInfo.LastProcessedInputFrame != INDEX_NONE, LogPhysics, Warning, TEXT("[Remote.Input] in fault. Reusing Inputcmd. (Client) Input: %d. (Server) Local Frame: %d"), FrameInfo.LastProcessedInputFrame, FrameInfo.LastLocalFrame);
-								return;
-							}
-							FrameInfo.bFault = false;
-						}
-						else if (NumBufferedInputCmds <= 0)
-						{
-							// No Cmds to process, enter fault state. Increment TargetNumBufferedCmds each time this happens.
-							// TODO: We should have something to bring this back down (which means skipping frames) we don't want temporary poor conditions to cause permanent high input buffering
-							FrameInfo.bFault = true;
-							FrameInfo.TargetNumBufferedCmds = FMath::Min(FrameInfo.TargetNumBufferedCmds + InputCmdCVars::TargetNumBufferedCmdsDeltaOnFault, InputCmdCVars::MaxTargetNumBufferedCmds);
-
-							UE_CLOG(FrameInfo.LastProcessedInputFrame != INDEX_NONE, LogPhysics, Warning, TEXT("[Remote.Input] ENTERING fault. New Target: %.2f. (Client) Input: %d. (Server) Local Frame: %d"), FrameInfo.TargetNumBufferedCmds, FrameInfo.LastProcessedInputFrame, FrameInfo.LastLocalFrame);
-							return;
-						}
-
-						float TargetTimeDilation = 1.f;
-						if (NumBufferedInputCmds < (int32)FrameInfo.TargetNumBufferedCmds)
-						{
-							TargetTimeDilation += InputCmdCVars::MaxTimeDilationMag; // Tell client to speed up, we are starved on cmds
-						}
-						/*else if (NumBufferedInputCmds > (int32)FrameInfo.TargetNumBufferedCmds)
-						{
-							TargetTimeDilation -= InputCmdCVars::MaxTimeDilationMag; // Tell client to slow down, we have too many buffered cmds
-
-							if (InputCmdCVars::LerpTargetNumBufferedCmdsAggresively == 0)
-							{
-								// When non aggressive, only lerp when we are above our limit
-								FrameInfo.TargetNumBufferedCmds = FMath::Lerp(FrameInfo.TargetNumBufferedCmds, InputCmdCVars::TargetNumBufferedCmds, InputCmdCVars::TargetNumBufferedCmdsAlpha);
-							}
-						}*/
-
-						FrameInfo.TargetTimeDilation = FMath::Lerp(FrameInfo.TargetTimeDilation, TargetTimeDilation, InputCmdCVars::TimeDilationAlpha);
-						FrameInfo.QuantizedTimeDilation = QuantizeTimeDilation(TargetTimeDilation);
-
-						if (InputCmdCVars::LerpTargetNumBufferedCmdsAggresively != 0)
-						{
-							// When aggressive, always lerp towards target
-							FrameInfo.TargetNumBufferedCmds = FMath::Lerp(FrameInfo.TargetNumBufferedCmds, InputCmdCVars::TargetNumBufferedCmds, InputCmdCVars::TargetNumBufferedCmdsAlpha);
-						}
-
-						FrameInfo.LastProcessedInputFrame++;
-						FrameInfo.LastLocalFrame = PhysicsStep;
-					}
-				}
-			}
-		}
-	}
-
-
-	// Updates the TMap on PhysScene that stores (non interpolated) physics data for replication.
-	// 
-	// Needs to be called from PT context to access fixed tick handle
-	// but also needs to be able to access GT data (actor iterator, actor state)
-	int32 LastUpdatedStep = 0;
-	void UpdateReplicationMap_Internal(int32 PhysicsStep)
-	{
-		if (LastUpdatedStep == PhysicsStep)
-		{
-			return;
-		}
-
-		LastUpdatedStep = PhysicsStep;
-
-		// Go through all "managed" primitive components and update Tmap to have latest physics state
-		// This is a temp hack and we will eventually replace this with something internal to the physics system that can be updated as things move
-		if (FPhysScene_Chaos* Scene = static_cast<FPhysScene_Chaos*>(World->GetPhysicsScene()))
-		{
-			Scene->ReplicationCache.ServerFrame = PhysicsStep;
-			for (TActorIterator<AActor> It(World); It; ++It)
-			{
-				AActor* Actor = *It;
-				if (Actor && Actor->GetIsReplicated() && Actor->IsReplicatingMovement())
-				{
-					UPrimitiveComponent* RootComponent = Cast<UPrimitiveComponent>(Actor->GetRootComponent());
-					if (RootComponent && RootComponent->IsSimulatingPhysics())
-					{
-						if (FBodyInstanceAsyncPhysicsTickHandle Handle = RootComponent->GetBodyInstanceAsyncPhysicsTickHandle())
-						{
-							FObjectKey Key(RootComponent);
-							FRigidBodyState& LatestState = Scene->ReplicationCache.Map.FindOrAdd(Key);
-							
-							// This might be wrong... see FBodyInstance::GetRigidBodyState (converts to unreal units?) and FRepMovement::FillFrom
-							LatestState.Position = Handle->X();
-							LatestState.Quaternion = Handle->R();
-							LatestState.LinVel = Handle->V();
-							LatestState.AngVel = Handle->W();
-							LatestState.Flags = Handle->ObjectState() == Chaos::EObjectStateType::Sleeping ? ERigidBodyFlags::Sleeping : 0;
-						}
-					}
-				}
 			}
 		}
 	}
@@ -366,12 +154,6 @@ void UAsyncPhysicsInputComponent::AsyncPhysicsTickComponent(float DeltaTime, flo
 	FAsyncPhysicsInputRewindCallback* Callback = static_cast<FAsyncPhysicsInputRewindCallback*>(Solver->GetRewindCallback());
 	const int32 ServerFrame = Callback->CachedServerFrame;
 
-	if(World->IsNetMode(ENetMode::NM_DedicatedServer) || World->IsNetMode(ENetMode::NM_ListenServer))
-	{
-		//TODO: move this somewhere else - here because we need to run this when GT and PT are both on same core. Function guards against multiple calls per ServerFrame so calling from each instance is ok
-		Callback->UpdateReplicationMap_Internal(ServerFrame);
-	}
-
 	DataToConsume = nullptr; //set current input to none in case we don't find it
 
 	for (int32 Idx = BufferedData.Num() - 1; Idx >= 0; --Idx)
diff --git a/Engine/Source/Runtime/Engine/Private/PhysicsEngine/Experimental/PhysScene_Chaos.cpp b/Engine/Source/Runtime/Engine/Private/PhysicsEngine/Experimental/PhysScene_Chaos.cpp
--- a/Engine/Source/Runtime/Engine/Private/PhysicsEngine/Experimental/PhysScene_Chaos.cpp
+++ b/Engine/Source/Runtime/Engine/Private/PhysicsEngine/Experimental/PhysScene_Chaos.cpp
@@ -16,6 +16,7 @@
 #include "PhysicsEngine/PhysicsCollisionHandler.h"
 #include "PhysicsEngine/PhysicsObjectExternalInterface.h"
 #include "Physics/Experimental/ChaosEventRelay.h"
+#include "EngineUtils.h"
 
 #include "ChaosSolversModule.h"
 
@@ -95,7 +96,6 @@ public:
 			}
 		}
 
-
 		const UPhysicsSettings* PhysicsSettings = UPhysicsSettings::Get();
 		const bool bEnableResim = PhysicsSettings->PhysicsPrediction.bEnablePhysicsResimulation;
 		const int32 NumFrames = PhysicsSettings->GetPhysicsHistoryCount();
@@ -1213,6 +1213,33 @@ void FPhysScene_Chaos::AddToComponentMaps(UPrimitiveComponent* Component, IPhysi
 	}
 }
 
+void FPhysScene_Chaos::PopulateReplicationCache(const int32 PhysicsStep)
+{
+	ReplicationCache.ServerFrame = PhysicsStep;
+	for (TActorIterator<AActor> It(GetOwningWorld()); It; ++It)
+	{
+		AActor* Actor = *It;
+		if (Actor && Actor->GetIsReplicated() && Actor->IsReplicatingMovement())
+		{
+			UPrimitiveComponent* RootComponent = Cast<UPrimitiveComponent>(Actor->GetRootComponent());
+			if (RootComponent && RootComponent->IsSimulatingPhysics())
+			{
+				if (FBodyInstanceAsyncPhysicsTickHandle Handle = RootComponent->GetBodyInstanceAsyncPhysicsTickHandle())
+				{
+					FObjectKey Key(RootComponent);
+					FRigidBodyState& ReplicationState = ReplicationCache.Map.FindOrAdd(Key);
+
+					ReplicationState.Position = Handle->X();
+					ReplicationState.Quaternion = Handle->R();
+					ReplicationState.LinVel = Handle->V();
+					ReplicationState.AngVel = Handle->W();
+					ReplicationState.Flags = Handle->ObjectState() == Chaos::EObjectStateType::Sleeping ? ERigidBodyFlags::Sleeping : 0;
+				}
+			}
+		}
+	}
+}
+
 void FPhysScene_Chaos::RemoveFromComponentMaps(IPhysicsProxyBase* InObject)
 {
 	auto* const Component = PhysicsProxyToComponentMap.Find(InObject);
diff --git a/Engine/Source/Runtime/Engine/Private/PhysicsEngine/NetworkPhysicsComponent.cpp b/Engine/Source/Runtime/Engine/Private/PhysicsEngine/NetworkPhysicsComponent.cpp
new file mode 100644
--- /dev/null
+++ b/Engine/Source/Runtime/Engine/Private/PhysicsEngine/NetworkPhysicsComponent.cpp
@@ -0,0 +1,865 @@
+// Copyright Epic Games, Inc. All Rights Reserved.
+
+#include "Physics/NetworkPhysicsComponent.h"
+
+#include "Components/PrimitiveComponent.h"
+#include "EngineLogs.h"
+#include "EngineUtils.h"
+#include "GameFramework/Pawn.h"
+#include "GameFramework/PlayerController.h"
+#include "PBDRigidsSolver.h"
+#include "Net/UnrealNetwork.h"
+#include "PhysicsReplication.h"
+#include "PhysicsProxy/SingleParticlePhysicsProxy.h"
+
+#include UE_INLINE_GENERATED_CPP_BY_NAME(NetworkPhysicsComponent)
+
+namespace InputCmdCVars
+{
+	static int32 ForceFault = 0;
+	static FAutoConsoleVariableRef CVarForceFault(TEXT("p.net.ForceFault"), ForceFault, TEXT("Forces server side input fault"));
+
+	static int32 MaxBufferedCmds = 16;
+	static FAutoConsoleVariableRef CVarMaxBufferedCmds(TEXT("p.net.MaxBufferedCmds"), MaxBufferedCmds, TEXT("MaxNumber of buffered server side commands"));
+
+	static int32 TimeDilationEnabled = 0;
+	static FAutoConsoleVariableRef CVarTimeDilationEnabled(TEXT("p.net.TimeDilationEnabled"), TimeDilationEnabled, TEXT("Enable clientside TimeDilation"));
+
+	static float MaxTargetNumBufferedCmds = 5.0;
+	static FAutoConsoleVariableRef CVarMaxTargetNumBufferedCmds(TEXT("p.net.MaxTargetNumBufferedCmds"), MaxTargetNumBufferedCmds, TEXT("Maximum number of buffered inputs the server will target per client."));
+
+	static float MaxTimeDilationMag = 0.01f;
+	static FAutoConsoleVariableRef CVarMaxTimeDilationMag(TEXT("p.net.MaxTimeDilationMag"), MaxTimeDilationMag, TEXT("Maximum time dilation that client will use to slow down / catch up with server"));
+
+	static float TimeDilationAlpha = 0.1f;
+	static FAutoConsoleVariableRef CVarTimeDilationAlpha(TEXT("p.net.TimeDilationAlpha"), TimeDilationAlpha, TEXT("Lerp strength for sliding client time dilation"));
+
+	static float TargetNumBufferedCmdsDeltaOnFault = 1.0f;
+	static FAutoConsoleVariableRef CVarTargetNumBufferedCmdsDeltaOnFault(TEXT("p.net.TargetNumBufferedCmdsDeltaOnFault"), TargetNumBufferedCmdsDeltaOnFault, TEXT("How much to increase TargetNumBufferedCmds when an input fault occurs"));
+
+	static float TargetNumBufferedCmds = 1.9f;
+	static FAutoConsoleVariableRef CVarTargetNumBufferedCmds(TEXT("p.net.TargetNumBufferedCmds"), TargetNumBufferedCmds, TEXT("How much to increase TargetNumBufferedCmds when an input fault occurs"));
+
+	static float TargetNumBufferedCmdsAlpha = 0.005f;
+	static FAutoConsoleVariableRef CVarTargetNumBufferedCmdsAlpha(TEXT("p.net.TargetNumBufferedCmdsAlpha"), TargetNumBufferedCmdsAlpha, TEXT("Lerp strength for TargetNumBufferedCmds"));
+
+	static int32 LerpTargetNumBufferedCmdsAggresively = 0;
+	static FAutoConsoleVariableRef CVarLerpTargetNumBufferedCmdsAggresively(TEXT("p.net.LerpTargetNumBufferedCmdsAggresively"), LerpTargetNumBufferedCmdsAggresively, TEXT("Aggresively lerp towards TargetNumBufferedCmds. Reduces server side buffering but can cause more artifacts."));
+}
+
+// --------------------------------------------------------------------------------------------------------------------------------------------------
+//	Client InputCmd Stream stuff
+// --------------------------------------------------------------------------------------------------------------------------------------------------
+
+namespace
+{
+
+int8 QuantizeTimeDilation(float F)
+{
+	if (F == 1.f)
+	{
+		return 0;
+	}
+
+	float Normalized = FMath::Clamp<float>((F - 1.f) / InputCmdCVars::MaxTimeDilationMag, -1.f, 1.f);
+	return (int8)(Normalized * 128.f);
+}
+
+float DeQuantizeTimeDilation(int8 i)
+{
+	if (i == 0)
+	{
+		return 1.f;
+	}
+
+	float Normalized = (float)i / 128.f;
+	float Uncompressed = 1.f + (Normalized * InputCmdCVars::MaxTimeDilationMag);
+	return Uncompressed;
+}
+
+}
+
+// after presimulate internal (asyncinput internal simulation done and the output created)
+void FNetworkPhysicsCallback::ApplyCallbacks_Internal(int32 PhysicsStep, const TArray<Chaos::ISimCallbackObject*>& SimCallbackObjects)
+{
+	if(World->IsNetMode(NM_ListenServer) || World->IsNetMode(NM_DedicatedServer))
+	{
+		if (FPhysScene_Chaos* Scene = static_cast<FPhysScene_Chaos*>(World->GetPhysicsScene()))
+		{
+			Scene->PopulateReplicationCache(PhysicsStep);
+		}
+	}
+}
+
+// before presimulate internal (asyncinput internal simulation is not done yet)
+void FNetworkPhysicsCallback::ProcessInputs_Internal(int32 PhysicsStep, const TArray<Chaos::FSimCallbackInputAndObject>& SimCallbacks)
+{
+	PreProcessInputsInternal.Broadcast(PhysicsStep);
+	for (Chaos::ISimCallbackObject* SimCallbackObject : RewindableCallbackObjects)
+	{
+		SimCallbackObject->ProcessInputs_Internal(PhysicsStep);
+	}
+	PostProcessInputsInternal.Broadcast(PhysicsStep);
+}
+
+int32 FNetworkPhysicsCallback::TriggerRewindIfNeeded_Internal(int32 LatestStepCompleted)
+{
+	int32 ResimFrame = INDEX_NONE;
+	for (Chaos::ISimCallbackObject* SimCallbackObject : RewindableCallbackObjects)
+	{
+		const int32 CallbackFrame = SimCallbackObject->TriggerRewindIfNeeded_Internal(LatestStepCompleted);
+		ResimFrame = (ResimFrame == INDEX_NONE) ? CallbackFrame : FMath::Min(CallbackFrame, ResimFrame);
+	}
+
+#if DEBUG_NETWORK_PHYSICS
+	UE_LOG(LogTemp, Log, TEXT("CLIENT | PT | TriggerRewindIfNeeded_Internal | Callbacks Frame = %d"), ResimFrame);
+#endif
+
+	if (World->GetNetMode() == NM_Client && World->GetPhysicsScene())
+	{
+		if (FPhysScene_Chaos* Scene = static_cast<FPhysScene_Chaos*>(World->GetPhysicsScene()))
+		{
+			if (const IPhysicsReplication* PhysicsReplication = Scene->GetPhysicsReplication())
+			{
+				const int32 ReplicationFrame = PhysicsReplication->GetResimFrame();
+
+#if DEBUG_NETWORK_PHYSICS
+				UE_LOG(LogTemp, Log, TEXT("CLIENT | PT | TriggerRewindIfNeeded_Internal | Replication Frame = %d"), ReplicationFrame);
+#endif
+				ResimFrame = (ResimFrame == INDEX_NONE) ? ReplicationFrame : (ReplicationFrame == INDEX_NONE) ? ResimFrame : FMath::Min(ReplicationFrame, ResimFrame);
+				Scene->GetPhysicsReplication()->SetResimFrame(INDEX_NONE);
+			}
+		}
+	}
+
+	if(ResimFrame != INDEX_NONE)
+	{ 
+		FPhysScene* PhysScene = World->GetPhysicsScene();
+		if (ensureAlways(PhysScene))
+		{
+			if (Chaos::FPhysicsSolver* Solver = PhysScene->GetSolver())
+			{
+				const int32 ValidFrame = Solver->GetRewindData()->FindValidResimFrame(ResimFrame);
+#if DEBUG_NETWORK_PHYSICS
+				UE_LOG(LogTemp, Log, TEXT("CLIENT | PT | TriggerRewindIfNeeded_Internal | Resim Frame = %d | Valid Frame = %d"), ResimFrame, ValidFrame);
+#endif
+				ResimFrame = ValidFrame;
+			}
+		}
+	}
+	
+	return ResimFrame;
+}
+
+void FNetworkPhysicsCallback::UpdateClientPlayer_External(int32 PhysicsStep)
+{
+	int32 LocalOffset = 0;
+	if (APlayerController* PC = World->GetFirstPlayerController())
+	{
+		// ------------------------------------------------
+		// Send RPC to server telling them what (client/local) physics step we are running
+		//	* Note that SendData is empty because of the existing API, should change this
+		// ------------------------------------------------	
+
+		TArray<uint8> SendData;
+		PC->PushClientInput(PhysicsStep, SendData);
+
+		// -----------------------------------------------------------------
+		// Calculate latest frame offset to map server frame to local frame
+		// -----------------------------------------------------------------
+		APlayerController::FClientFrameInfo& ClientFrameInfo = PC->GetClientFrameInfo();
+
+		// -----------------------------------------------------------------
+		// Apply local TIme Dilation based on server's recommendation.
+		// This speeds up or slows down our consumption of real time (by like < 1%)
+		// Ultimately this causes us to send InputCmds at a lower or higher rate in
+		// order to keep server side buffer at optimal capacity. 
+		// Optimal capacity = as small as possible without ever "missing" a frame (e.g, minimal buffer yet always a new fresh cmd to consume server side)
+		// -----------------------------------------------------------------
+		const float RealTimeDilation = DeQuantizeTimeDilation(ClientFrameInfo.QuantizedTimeDilation);
+
+		if (InputCmdCVars::TimeDilationEnabled > 0)
+		{
+			FPhysScene* PhysScene = World->GetPhysicsScene();
+			PhysScene->SetNetworkDeltaTimeScale(RealTimeDilation);
+		}
+	}
+}
+
+void FNetworkPhysicsCallback::UpdateServerPlayer_External(int32 PhysicsStep)
+{
+	// -----------------------------------------------
+	// Server: "consume" an InputCmd from each Player Controller
+	// All this means in this context is updating FServerFrameInfo:: LastProcessedInputFrame, LastLocalFrame
+	// (E.g, telling each client what "Input" of theirs we were processing and our local physics frame number.
+	// In cases where the buffer has a fault, we calculate a suggested time dilation to temporarily make client speed up 
+	// or slow down their input cmd production.
+	// -----------------------------------------------
+
+	const bool bForceFault = InputCmdCVars::ForceFault > 0;
+	InputCmdCVars::ForceFault = FMath::Max(0, InputCmdCVars::ForceFault - 1);
+
+	for (FConstPlayerControllerIterator Iterator = World->GetPlayerControllerIterator(); Iterator; ++Iterator)
+	{
+		if (APlayerController* PC = Iterator->Get())
+		{
+			APlayerController::FServerFrameInfo& FrameInfo = PC->GetServerFrameInfo();
+			APlayerController::FInputCmdBuffer& InputBuffer = PC->GetInputBuffer();
+
+			{
+				const int32 NumBufferedInputCmds = bForceFault ? 0 : (InputBuffer.HeadFrame() - FrameInfo.LastProcessedInputFrame);
+
+				// Check Overflow
+				if (NumBufferedInputCmds > InputCmdCVars::MaxBufferedCmds)
+				{
+					UE_LOG(LogPhysics, Warning, TEXT("[Remote.Input] overflow %d %d -> %d"), InputBuffer.HeadFrame(), FrameInfo.LastProcessedInputFrame, NumBufferedInputCmds);
+					FrameInfo.LastProcessedInputFrame = InputBuffer.HeadFrame() - InputCmdCVars::MaxBufferedCmds + 1;
+				}
+
+				// Check fault - we are waiting for Cmds to reach TargetNumBufferedCmds before continuing
+				if (FrameInfo.bFault)
+				{
+					if (NumBufferedInputCmds < (int32)FrameInfo.TargetNumBufferedCmds)
+					{
+						// Skip this because it is in fault. We will use the prev input for this frame.
+						UE_CLOG(FrameInfo.LastProcessedInputFrame != INDEX_NONE, LogPhysics, Warning, TEXT("[Remote.Input] in fault. Reusing Inputcmd. (Client) Input: %d. (Server) Local Frame: %d"), FrameInfo.LastProcessedInputFrame, FrameInfo.LastLocalFrame);
+						continue;
+					}
+					FrameInfo.bFault = false;
+				}
+				else if (NumBufferedInputCmds <= 0)
+				{
+					// No Cmds to process, enter fault state. Increment TargetNumBufferedCmds each time this happens.
+					// TODO: We should have something to bring this back down (which means skipping frames) we don't want temporary poor conditions to cause permanent high input buffering
+					FrameInfo.bFault = true;
+					FrameInfo.TargetNumBufferedCmds = FMath::Min(FrameInfo.TargetNumBufferedCmds + InputCmdCVars::TargetNumBufferedCmdsDeltaOnFault, InputCmdCVars::MaxTargetNumBufferedCmds);
+
+					UE_CLOG(FrameInfo.LastProcessedInputFrame != INDEX_NONE, LogPhysics, Warning, TEXT("[Remote.Input] ENTERING fault. New Target: %.2f. (Client) Input: %d. (Server) Local Frame: %d"), FrameInfo.TargetNumBufferedCmds, FrameInfo.LastProcessedInputFrame, FrameInfo.LastLocalFrame);
+					continue;
+				}
+
+				float TargetTimeDilation = 1.f;
+				if (NumBufferedInputCmds < (int32)FrameInfo.TargetNumBufferedCmds)
+				{
+					TargetTimeDilation += InputCmdCVars::MaxTimeDilationMag; // Tell client to speed up, we are starved on cmds
+				}
+
+				FrameInfo.TargetTimeDilation = FMath::Lerp(FrameInfo.TargetTimeDilation, TargetTimeDilation, InputCmdCVars::TimeDilationAlpha);
+				FrameInfo.QuantizedTimeDilation = QuantizeTimeDilation(TargetTimeDilation);
+
+				if (InputCmdCVars::LerpTargetNumBufferedCmdsAggresively != 0)
+				{
+					// When aggressive, always lerp towards target
+					FrameInfo.TargetNumBufferedCmds = FMath::Lerp(FrameInfo.TargetNumBufferedCmds, InputCmdCVars::TargetNumBufferedCmds, InputCmdCVars::TargetNumBufferedCmdsAlpha);
+				}
+
+				FrameInfo.LastProcessedInputFrame++;
+				FrameInfo.LastLocalFrame = PhysicsStep;
+			}
+		}
+	}
+}
+
+void FNetworkPhysicsCallback::ProcessInputs_External(int32 PhysicsStep, const TArray<Chaos::FSimCallbackInputAndObject>& SimCallbackInputs)
+{
+	for(const Chaos::FSimCallbackInputAndObject& SimCallbackObject : SimCallbackInputs)
+	{
+		if (SimCallbackObject.CallbackObject && SimCallbackObject.CallbackObject->HasOption(Chaos::ESimCallbackOptions::Rewind))
+		{
+			SimCallbackObject.CallbackObject->ProcessInputs_External(PhysicsStep);
+		}
+	}
+
+	if (World->GetNetMode() == NM_Client)
+	{
+		UpdateClientPlayer_External(PhysicsStep);
+	}
+	else
+	{
+		UpdateServerPlayer_External(PhysicsStep);
+	}
+}
+
+UNetworkPhysicsSystem::UNetworkPhysicsSystem()
+{}
+
+void UNetworkPhysicsSystem::Initialize(FSubsystemCollectionBase& Collection)
+{
+	UWorld* World = GetWorld();
+	check(World);
+
+	if (World->WorldType == EWorldType::PIE || World->WorldType == EWorldType::Game)
+	{
+		FWorldDelegates::OnPostWorldInitialization.AddUObject(this, &UNetworkPhysicsSystem::OnWorldPostInit);
+	}
+}
+
+void UNetworkPhysicsSystem::Deinitialize()
+{}
+
+void UNetworkPhysicsSystem::OnWorldPostInit(UWorld* World, const UWorld::InitializationValues)
+{
+	if (World != GetWorld())
+	{
+		return;
+	}
+
+	if(UPhysicsSettings::Get()->PhysicsPrediction.bEnablePhysicsPrediction)
+	{
+		if (FPhysScene* PhysScene = World->GetPhysicsScene())
+		{
+			if(Chaos::FPhysicsSolver* Solver = PhysScene->GetSolver())
+			{ 
+				if (Solver->GetRewindCallback() == nullptr)
+				{
+					const int32 NumFrames = FMath::Max<int32>(1, UPhysicsSettings::Get()->GetPhysicsHistoryCount());
+					Solver->EnableRewindCapture(NumFrames, true, MakeUnique<FNetworkPhysicsCallback>(World));
+				}
+			}
+		}
+	}
+}
+
+UNetworkPhysicsComponent::UNetworkPhysicsComponent(const FObjectInitializer& ObjectInitializer)
+	: Super(ObjectInitializer)
+{
+	InitPhysics();
+}
+
+UNetworkPhysicsComponent::UNetworkPhysicsComponent() : Super()
+{
+	InitPhysics();
+}
+
+void UNetworkPhysicsComponent::InitPhysics()
+{
+	bAutoActivate = true;
+	bWantsInitializeComponent = true;
+	SetIsReplicatedByDefault(true);
+	SetAsyncPhysicsTickEnabled(true);
+
+	StatesOffsets.SetNumZeroed(StatesRedundancy + 1);
+	InputsOffsets.SetNumZeroed(InputsRedundancy + 1);
+
+	if (APawn* Pawn = Cast<APawn>(GetOwner()))
+	{
+		FRepMovement& RepMovement = Pawn->GetReplicatedMovement_Mutable();
+		RepMovement.LocationQuantizationLevel = EVectorQuantization::RoundTwoDecimals;
+		RepMovement.RotationQuantizationLevel = ERotatorQuantization::ShortComponents;
+		RepMovement.VelocityQuantizationLevel = EVectorQuantization::RoundTwoDecimals;
+	}
+}
+
+void UNetworkPhysicsComponent::BeginPlay()
+{
+	Super::BeginPlay();
+	UWorld* World = GetWorld();
+
+	if (World->WorldType == EWorldType::PIE || World->WorldType == EWorldType::Game)
+	{
+		if (FPhysScene* PhysScene = World->GetPhysicsScene())
+		{
+			if (Chaos::FPhysicsSolver* Solver = PhysScene->GetSolver())
+			{
+				if(FNetworkPhysicsCallback* SolverCallback = static_cast<FNetworkPhysicsCallback*>(Solver->GetRewindCallback()))
+				{
+					SolverCallback->PreProcessInputsInternal.AddUObject(this, &UNetworkPhysicsComponent::OnPreProcessInputsInternal);
+					SolverCallback->PostProcessInputsInternal.AddUObject(this, &UNetworkPhysicsComponent::OnPostProcessInputsInternal);
+				}
+			}
+		}
+	}
+}
+
+void UNetworkPhysicsComponent::InitializeComponent()
+{
+	Super::InitializeComponent();
+	UWorld* World = GetWorld();
+
+	if (UNetworkPhysicsSystem* NetworkManager = World->GetSubsystem<UNetworkPhysicsSystem>())
+	{
+		NetworkManager->RegisterNetworkComponent(this);
+	}
+}
+
+void UNetworkPhysicsComponent::UninitializeComponent()
+{
+	Super::UninitializeComponent();
+	UWorld* World = GetWorld();
+
+	if (UNetworkPhysicsSystem* NetworkManager = World->GetSubsystem<UNetworkPhysicsSystem>())
+	{
+		NetworkManager->UnregisterNetworkComponent(this);
+	}
+}
+
+void UNetworkPhysicsComponent::GetLifetimeReplicatedProps(TArray< FLifetimeProperty >& OutLifetimeProps) const
+{
+	Super::GetLifetimeReplicatedProps(OutLifetimeProps);
+
+	DOREPLIFETIME_CONDITION_NOTIFY(UNetworkPhysicsComponent, ReplicatedInputs, COND_None, REPNOTIFY_Always);
+	DOREPLIFETIME_CONDITION_NOTIFY(UNetworkPhysicsComponent, ReplicatedStates, COND_None, REPNOTIFY_Always);
+}
+
+void UNetworkPhysicsComponent::UpdatePackageMap()
+{
+	APlayerController* Controller = GetPlayerController();
+	UPackageMap* PackageMap = (Controller && Controller->GetNetConnection()) ? Controller->GetNetConnection()->PackageMap : nullptr;
+
+	if(InputsHistory && StatesHistory)
+	{
+		InputsHistory->SetPackageMap(PackageMap);
+		StatesHistory->SetPackageMap(PackageMap);
+	}
+}
+
+void UNetworkPhysicsComponent::AsyncPhysicsTickComponent(float DeltaTime, float SimTime)
+{
+	Super ::AsyncPhysicsTickComponent(DeltaTime, SimTime);
+#if DEBUG_NETWORK_PHYSICS
+	if(HasServerWorld() && !HasLocalController() && InputsHistory)
+	{
+		TArray<int32> LocalFrames, ServerFrames, InputFrames;
+		InputsHistory->DebugDatas(ReplicatedInputs, LocalFrames, ServerFrames, InputFrames);
+
+		UE_LOG(LogTemp, Log, TEXT("SERVER | PT | AsyncPhysicsTickComponent | Receiving %d inputs from CLIENT | Component = %s"), LocalFrames.Num(), *GetFullName());
+		for (int32 FrameIndex = 0; FrameIndex < LocalFrames.Num(); ++FrameIndex)
+		{
+			UE_LOG(LogTemp, Log, TEXT("		Debugging replicated inputs at local frame = %d | server frame = %d | Component = %s"),
+				LocalFrames[FrameIndex], ServerFrames[FrameIndex], *GetFullName());
+		}
+	}
+#endif
+	// Update the package map for serialization
+	UpdatePackageMap();
+
+	// Record the received states from the server into the history for future use
+	if (UWorld* World = GetWorld())
+	{
+		if (FPhysScene* PhysScene = World->GetPhysicsScene())
+		{
+			if(!PhysScene->GetSolver()->GetEvolution()->IsResimming())
+			{
+				// Send the inputs across the network
+				SendLocalInputsDatas();
+
+				// Send the states across the network
+				SendLocalStatesDatas();
+
+				// Advance the NetworkIndex
+				InputsIndex = (InputsIndex + 1) % (InputsRedundancy + 1);
+				StatesIndex = (StatesIndex + 1) % (StatesRedundancy + 1);
+			}
+		}
+	}
+}
+
+void UNetworkPhysicsComponent::SendLocalInputsDatas()
+{
+	if (HasLocalController() && InputsHistory)
+	{
+		// We just check that the local client to server offset is valid before doing something
+		const int32 LocalOffset = HasServerWorld() ? 0 : GetPlayerController()->GetLocalToServerAsyncPhysicsTickOffset();
+
+		if (LocalOffset >= 0)
+		{
+			const int32 NextIndex = (InputsIndex + 1) % (InputsRedundancy + 1);
+			InputsHistory->SerializeDatas(InputsOffsets[NextIndex], InputsOffsets[InputsIndex], LocalInputs, LocalOffset);
+
+			if (HasServerWorld())
+			{
+				// if on server (Listen server) we should send the inputs onto all the clients through repnotify
+				ReplicatedInputs = LocalInputs;
+			}
+			else
+			{
+#if DEBUG_NETWORK_PHYSICS
+				if (APlayerController* PlayerController = GetPlayerController())
+				{
+					FAsyncPhysicsTimestamp Timestamp = PlayerController->GetAsyncPhysicsTimestamp();
+
+					TArray<int32> LocalFrames, ServerFrames, InputFrames;
+					InputsHistory->DebugDatas(LocalInputs, LocalFrames, ServerFrames, InputFrames);
+
+					UE_LOG(LogTemp, Log, TEXT("CLIENT | GT | SendLocalInputsDatas | Sending %d inputs from CLIENT | Component = %s"), LocalFrames.Num(), *GetFullName());
+					for (int32 FrameIndex = 0; FrameIndex < LocalFrames.Num(); ++FrameIndex)
+					{
+						UE_LOG(LogTemp, Log, TEXT("		Debugging local inputs at local frame = %d | server frame = %d | Current local frame = %d | Current server frame = %d"),
+							LocalFrames[FrameIndex], ServerFrames[FrameIndex], Timestamp.LocalFrame, Timestamp.ServerFrame);
+					}
+				}
+#endif
+
+				// if on the client we should first send the replicated inputs onto the server
+				// the RPC will then resend them onto all the other clients (except the local ones)
+				ServerReceiveInputsDatas(LocalInputs);
+			}
+		}
+	}
+}
+
+void UNetworkPhysicsComponent::SendLocalStatesDatas()
+{
+	if (HasServerWorld() && StatesHistory)
+	{
+		const int32 NextIndex = (StatesIndex + 1) % (StatesRedundancy + 1);
+		StatesHistory->SerializeDatas(StatesOffsets[NextIndex], StatesOffsets[StatesIndex], LocalStates, 0);
+
+		// if on server we should send the states onto all the clients through repnotify
+		ReplicatedStates = LocalStates;
+	}
+}
+
+void UNetworkPhysicsComponent::CorrectServerToLocalOffset(const int32 LocalToServerOffset)
+{
+	if (HasLocalController() && !HasServerWorld() && StatesHistory)
+	{
+		const TArray<uint8> ReceivedStates = ReplicatedStates;
+
+		TArray<int32> LocalFrames, ServerFrames, InputFrames;
+		StatesHistory->DebugDatas(ReceivedStates, LocalFrames, ServerFrames, InputFrames);
+
+		int32 ServerToLocalOffset = LocalToServerOffset;
+		for (int32 FrameIndex = 0; FrameIndex < LocalFrames.Num(); ++FrameIndex)
+		{
+#if DEBUG_NETWORK_PHYSICS
+			UE_LOG(LogTemp, Log, TEXT("CLIENT | GT | CorrectServerToLocalOffset | Server frame = %d | Client Frame = %d"), ServerFrames[FrameIndex], InputFrames[FrameIndex]);
+#endif
+			ServerToLocalOffset = FMath::Min(ServerToLocalOffset, ServerFrames[FrameIndex] - InputFrames[FrameIndex]);
+		}
+
+		GetPlayerController()->SetServerToLocalAsyncPhysicsTickOffset(ServerToLocalOffset);
+#if DEBUG_NETWORK_PHYSICS
+		UE_LOG(LogTemp, Log, TEXT("CLIENT | GT | CorrectServerToLocalOffset | Server to local offset = %d | Local to server offset = %d"), ServerToLocalOffset, LocalToServerOffset);
+#endif
+	}
+}
+
+void UNetworkPhysicsComponent::OnRep_SetReplicatedStates()
+{
+	// The replicated states should only be used on the client since the server already have authoritative local ones
+	if (!HasServerWorld() && StatesHistory)
+	{
+		APlayerController* PlayerController = GetPlayerController();
+		if (!PlayerController)
+		{
+			PlayerController = GetWorld()->GetFirstPlayerController();
+		}
+		CorrectServerToLocalOffset(PlayerController->GetLocalToServerAsyncPhysicsTickOffset());
+		const int32 LocalOffset = PlayerController->GetServerToLocalAsyncPhysicsTickOffset();
+
+		// Record the received states from the server into the history for future use
+		if (UWorld* World = GetWorld())
+		{
+			if (FPhysScene* PhysScene = World->GetPhysicsScene())
+			{
+				const TArray<uint8> ReceivedStates = ReplicatedStates;
+				PhysScene->EnqueueAsyncPhysicsCommand(0, this, [this, PhysScene, ReceivedStates, LocalOffset]()
+				{
+					StatesHistory->DeserializeDatas(ReceivedStates, LocalOffset);
+#if DEBUG_NETWORK_PHYSICS
+					{
+						TArray<int32> LocalFrames, ServerFrames, InputFrames;
+						StatesHistory->DebugDatas(ReceivedStates, LocalFrames, ServerFrames, InputFrames);
+
+						UE_LOG(LogTemp, Log, TEXT("CLIENT | PT | OnRep_SetReplicatedStates | Receiving %d states from SERVER | Local offset = %d | Component = %s "), LocalFrames.Num(), LocalOffset, *GetFullName());
+						for (int32 FrameIndex = 0; FrameIndex < LocalFrames.Num(); ++FrameIndex)
+						{
+							UE_LOG(LogTemp, Log, TEXT("		Recording replicated states at local frame = %d | server frame = %d | life time = %d | Component = %s"), ServerFrames[FrameIndex] - LocalOffset, ServerFrames[FrameIndex], InputFrames[FrameIndex], *GetFullName());
+
+							if (HasLocalController() && (InputFrames[FrameIndex] != (ServerFrames[FrameIndex] - LocalOffset)))
+							{
+								UE_LOG(LogTemp, Log, TEXT("		Bad local frame compared to input frame!!!"));
+							}
+						}
+					}
+#endif
+				}, false);
+			}
+		}
+	}
+}
+
+void UNetworkPhysicsComponent::OnRep_SetReplicatedInputs()
+{
+	// For local controller we should already have correct replicated inputs
+	if (!HasLocalController() && !HasServerWorld() && InputsHistory)
+	{
+		APlayerController* PlayerController = GetPlayerController();
+		if(!PlayerController)
+		{
+			PlayerController = GetWorld()->GetFirstPlayerController();
+		}
+		const int32 LocalOffset = PlayerController->GetServerToLocalAsyncPhysicsTickOffset();
+
+		// Record the received inputs from the server into the history for future use
+		if (UWorld* World = GetWorld())
+		{
+			if (FPhysScene* PhysScene = World->GetPhysicsScene())
+			{
+				const TArray<uint8> ReceivedInputs = ReplicatedInputs;
+				PhysScene->EnqueueAsyncPhysicsCommand(0, this, [this, PhysScene, ReceivedInputs, LocalOffset]()
+				{
+					InputsHistory->DeserializeDatas(ReceivedInputs, LocalOffset);
+
+#if DEBUG_NETWORK_PHYSICS
+					{
+						TArray<int32> LocalFrames, ServerFrames, InputFrames;
+						InputsHistory->DebugDatas(ReceivedInputs, LocalFrames, ServerFrames, InputFrames);
+
+						UE_LOG(LogTemp, Log, TEXT("CLIENT | PT | OnRep_SetReplicatedInputs | Receiving %d inputs from SERVER | Local offset = %d | Component = %s"), LocalFrames.Num(), LocalOffset, *GetFullName());
+						for (int32 FrameIndex = 0; FrameIndex < LocalFrames.Num(); ++FrameIndex)
+						{
+							UE_LOG(LogTemp, Log, TEXT("		Recording replicated inputs at local frame = %d | server frame = %d | Component = %s"), ServerFrames[FrameIndex] - LocalOffset, ServerFrames[FrameIndex], *GetFullName());
+						}
+					}
+#endif
+				}, false);
+			}
+		}
+	}
+}
+
+bool UNetworkPhysicsComponent::ServerReceiveInputsDatas_Validate(const TArray<uint8>& ClientInputs)
+{
+	return true;
+}
+
+void UNetworkPhysicsComponent::ServerReceiveInputsDatas_Implementation(const TArray<uint8>& ClientInputs)
+{
+	if(InputsHistory)
+	{ 
+		// We could probably skip that test since the server RPC is on server
+		ensure(HasServerWorld());
+
+		// We could probably skip that test since the server RPC is on server
+		ensure(!HasLocalController());
+
+		ReplicatedInputs = ClientInputs;
+
+		// Record the received inputs from the client into the history for future use
+		if (UWorld* World = GetWorld())
+		{
+			if (FPhysScene* PhysScene = World->GetPhysicsScene())
+			{
+				const TArray<uint8> ReceivedInputs = ClientInputs;
+				PhysScene->EnqueueAsyncPhysicsCommand(0, this, [this, ReceivedInputs, PhysScene]()
+				{
+					InputsHistory->DeserializeDatas(ReceivedInputs, 0);
+	#if DEBUG_NETWORK_PHYSICS
+					{
+						TArray<int32> LocalFrames, ServerFrames, InputFrames;
+						InputsHistory->DebugDatas(ReceivedInputs, LocalFrames, ServerFrames, InputFrames);
+
+						const int32 CurrentFrame = PhysScene->GetSolver()->GetCurrentFrame();
+
+						const int32 EvalOffset = CurrentFrame - InputFrames[InputFrames.Num()-1] + 4;
+						UE_LOG(LogTemp, Log, TEXT("SERVER | PT | ServerReceiveInputsDatas | Receiving %d inputs from CLIENT | Inputs frame = %d | Server frame = %d | Eval Offset = %d | Component = %s | Num Bits = %d"), LocalFrames.Num(), InputFrames[InputFrames.Num() - 1], CurrentFrame, EvalOffset, *GetFullName(), ReceivedInputs.Num() * 8);
+						for (int32 FrameIndex = 0; FrameIndex < LocalFrames.Num(); ++FrameIndex)
+						{
+							UE_LOG(LogTemp, Log, TEXT("		Recording replicated inputs at local frame = %d | server frame = %d | Solver offset = %d | Component = %s"), 
+								LocalFrames[FrameIndex], ServerFrames[FrameIndex], ServerFrames[FrameIndex] - LocalFrames[FrameIndex], *GetFullName());
+						}
+					}
+	#endif
+				}, false);
+			}
+		}
+	}
+}
+
+void UNetworkPhysicsComponent::OnPreProcessInputsInternal(const int32 PhysicsStep)
+{
+#if DEBUG_NETWORK_PHYSICS
+	if (HasServerWorld())
+	{
+		UE_LOG(LogTemp, Log, TEXT("SERVER | PT | OnPreProcessInputsInternal | At Frame %d | Component = %s"), PhysicsStep, *GetFullName());
+	}
+	else
+	{
+		UE_LOG(LogTemp, Log, TEXT("CLIENT | PT | OnPreProcessInputsInternal | At Frame %d | Component = %s"), PhysicsStep, *GetFullName());
+	}
+#endif
+
+	if(InputsHistory && StatesHistory && ActorComponent)
+	{ 
+		bool bIsSolverReset = false;
+		bool bIsSolverResim = false;
+
+		if (FPhysScene* PhysScene = GetWorld()->GetPhysicsScene())
+		{ 
+			if (Chaos::FPhysicsSolver* Solver = PhysScene->GetSolver())
+			{
+				bIsSolverResim = Solver->GetEvolution()->IsResimming();
+				bIsSolverReset = Solver->GetEvolution()->IsResetting();
+			}
+		}
+
+		// for the inputs client local ones are ground truth otherwise use the replicated ones coming from the server
+		if (!HasLocalController() || bIsSolverResim)
+		{
+			FNetworkPhysicsDatas* PhysicsDatas = InputsDatas.Get();
+			PhysicsDatas->LocalFrame = PhysicsStep;
+	#if DEBUG_NETWORK_PHYSICS
+			UE_LOG(LogTemp, Log, TEXT("		Extracting history inputs at frame %d | Component = %s"), PhysicsStep, *GetFullName());
+	#endif
+			if (InputsHistory->ExtractDatas(PhysicsStep, bIsSolverReset, PhysicsDatas))
+			{ 
+				PhysicsDatas->ApplyDatas(ActorComponent);
+			}
+		}
+
+		if (!HasServerWorld() && bIsSolverResim)
+		{
+			FNetworkPhysicsDatas* PhysicsDatas = StatesDatas.Get();
+			PhysicsDatas->LocalFrame = PhysicsStep;
+	#if DEBUG_NETWORK_PHYSICS
+			UE_LOG(LogTemp, Log, TEXT("		Extracting history states at frame %d | Component = %s"), PhysicsStep, *GetFullName());
+	#endif
+			if (StatesHistory->ExtractDatas(PhysicsStep, bIsSolverReset, PhysicsDatas, true))
+			{
+				PhysicsDatas->ApplyDatas(ActorComponent);
+			}
+		}
+	}
+}
+
+void UNetworkPhysicsComponent::OnPostProcessInputsInternal(const int32 PhysicsStep)
+{
+	bool bIsSolverReset = false;
+	bool bIsSolverResim = false;
+#if DEBUG_NETWORK_PHYSICS
+	if (HasServerWorld())
+	{
+		UE_LOG(LogTemp, Log, TEXT("SERVER | PT | OnPostProcessInputsInternal | At Frame %d | Component = %s"), PhysicsStep, *GetFullName());
+	}
+	else
+	{
+		UE_LOG(LogTemp, Log, TEXT("CLIENT | PT | OnPostProcessInputsInternal | At Frame %d | Component = %s"), PhysicsStep, *GetFullName());
+	}
+#endif
+
+	if(InputsHistory && StatesHistory && ActorComponent)
+	{
+		if (FPhysScene* PhysScene = GetWorld()->GetPhysicsScene())
+		{
+			if (Chaos::FPhysicsSolver* Solver = PhysScene->GetSolver())
+			{
+				bIsSolverResim = Solver->GetEvolution()->IsResimming();
+				bIsSolverReset = Solver->GetEvolution()->IsResetting();
+			}
+		}
+
+		// for the inputs client local ones are ground truth otherwise use the replicated ones coming from the server
+		if (HasLocalController() && !bIsSolverResim && (InputsDatas != nullptr))
+		{
+			FNetworkPhysicsDatas* PhysicsDatas = InputsDatas.Get();
+			PhysicsDatas->LocalFrame = PhysicsStep;
+			PhysicsDatas->ServerFrame = HasServerWorld() ? PhysicsStep : PhysicsStep + GetPlayerController()->GetLocalToServerAsyncPhysicsTickOffset();
+			PhysicsDatas->InputFrame = PhysicsStep;
+
+			PhysicsDatas->BuildDatas(ActorComponent);
+
+			InputsOffsets[InputsIndex] = FMath::Max(InputsOffsets[InputsIndex], PhysicsStep + 1);
+			InputsHistory->RecordDatas(PhysicsStep, PhysicsDatas);
+
+	#if DEBUG_NETWORK_PHYSICS
+			UE_LOG(LogTemp, Log, TEXT("		Recording local inputs at frame %d | Component = %s"), PhysicsDatas->LocalFrame, *GetFullName());
+	#endif
+		}
+		// Compute of the local frame coming from the client that was used to generate this state
+		int32 InputFrame = INDEX_NONE;
+		{
+			FNetworkPhysicsDatas* PhysicsDatas = InputsDatas.Get();
+			if(InputsHistory->ExtractDatas(PhysicsStep, false, PhysicsDatas, true))
+			{
+				InputFrame = PhysicsDatas->InputFrame;
+			}
+		}
+
+		if (HasServerWorld() && !bIsSolverResim)
+		{
+			FNetworkPhysicsDatas* PhysicsDatas = StatesDatas.Get();
+			PhysicsDatas->LocalFrame = PhysicsStep;
+			PhysicsDatas->ServerFrame = PhysicsStep;
+			PhysicsDatas->InputFrame = InputFrame;
+
+			PhysicsDatas->BuildDatas(ActorComponent);
+
+			StatesOffsets[StatesIndex] = FMath::Max(StatesOffsets[StatesIndex], PhysicsStep+1);
+			StatesHistory->RecordDatas(PhysicsStep, PhysicsDatas);
+
+	#if DEBUG_NETWORK_PHYSICS
+			UE_LOG(LogTemp, Log, TEXT("		Recording local states at frame %d | from input frame = %d | Component = %s"), PhysicsDatas->LocalFrame, PhysicsDatas->InputFrame, *GetFullName());
+	#endif
+		}
+	}
+}
+
+bool UNetworkPhysicsComponent::HasServerWorld() const
+{
+	return GetWorld()->IsNetMode(NM_DedicatedServer) || GetWorld()->IsNetMode(NM_ListenServer);
+}
+
+bool UNetworkPhysicsComponent::HasLocalController() const
+{
+	if (APlayerController* PlayerController = GetPlayerController())
+	{
+		return PlayerController->IsLocalController();
+	}
+	return false;
+}
+
+APlayerController* UNetworkPhysicsComponent::GetPlayerController() const
+{
+	if (APlayerController* PC = Cast<APlayerController>(GetOwner()))
+	{
+		return PC;
+	}
+
+	if (APawn* Pawn = Cast<APawn>(GetOwner()))
+	{
+		return Pawn->GetController<APlayerController>();
+	}
+
+	return nullptr;
+}
+
+void UNetworkPhysicsComponent::RemoveDatasHistory()
+{
+	if (GetWorld())
+	{
+		if (FPhysScene* PhysScene = GetWorld()->GetPhysicsScene())
+		{
+			if (Chaos::FPhysicsSolver* Solver = PhysScene->GetSolver())
+			{
+				if (Chaos::FRewindData* RewindData = Solver->GetRewindData())
+				{
+					RewindData->RemoveInputsHistory(InputsHistory);
+					RewindData->RemoveStatesHistory(StatesHistory);
+				}
+			}
+		}
+	}
+}
+void UNetworkPhysicsComponent::AddDatasHistory()
+{
+	if (FPhysScene* PhysScene = GetWorld()->GetPhysicsScene())
+	{
+		if (Chaos::FPhysicsSolver* Solver = PhysScene->GetSolver())
+		{
+			if (Chaos::FRewindData* RewindData = Solver->GetRewindData())
+			{
+				RewindData->AddInputsHistory(InputsHistory);
+				RewindData->AddStatesHistory(StatesHistory);
+			}
+		}
+	}
+}
+
+
+
+
+
+
+
diff --git a/Engine/Source/Runtime/Engine/Private/PhysicsEngine/PhysicsReplication.cpp b/Engine/Source/Runtime/Engine/Private/PhysicsEngine/PhysicsReplication.cpp
--- a/Engine/Source/Runtime/Engine/Private/PhysicsEngine/PhysicsReplication.cpp
+++ b/Engine/Source/Runtime/Engine/Private/PhysicsEngine/PhysicsReplication.cpp
@@ -118,7 +118,7 @@ void ComputeDeltas(const FVector& CurrentPos, const FQuat& CurrentQuat, const FV
 	OutLinDiff = TargetPos - CurrentPos;
 	OutLinDiffSize = OutLinDiff.Size();
 	const FQuat InvCurrentQuat = CurrentQuat.Inverse();
-	const FQuat DeltaQuat = InvCurrentQuat * TargetQuat;
+	const FQuat DeltaQuat = TargetQuat * InvCurrentQuat;
 	DeltaQuat.ToAxisAndAngle(OutAngDiffAxis, OutAngDiff);
 	OutAngDiff = FMath::RadiansToDegrees(FMath::UnwindRadians(OutAngDiff));
 	OutAngDiffSize = FMath::Abs(OutAngDiff);
@@ -632,7 +632,7 @@ void FPhysicsReplication::OnTick(float DeltaSeconds, TMap<TWeakObjectPtr<UPrimit
 
 				if (FPhysicsSolverBase::IsNetworkPhysicsPredictionEnabled())
 				{
-					LocalFrameOffset = PlayerController->GetLocalToServerAsyncPhysicsTickOffset();
+					LocalFrameOffset = PlayerController->GetServerToLocalAsyncPhysicsTickOffset();
 				}
 				//TODO: as we send physics updates down we need to record latest seen
 				//NumPredictedFrames = Solver->GetCurrentFrame() - ClientFrameInfo.LastProcessedInputFrame;
diff --git a/Engine/Source/Runtime/Engine/Private/PlayerController.cpp b/Engine/Source/Runtime/Engine/Private/PlayerController.cpp
--- a/Engine/Source/Runtime/Engine/Private/PlayerController.cpp
+++ b/Engine/Source/Runtime/Engine/Private/PlayerController.cpp
@@ -66,6 +66,7 @@
 #include "LevelUtils.h"
 #include "WorldPartition/WorldPartitionSubsystem.h"
 #include "Physics/AsyncPhysicsInputComponent.h"
+#include "Physics/NetworkPhysicsComponent.h"
 #include "PBDRigidsSolver.h"
 #include "PhysicsEngine/PhysicsSettings.h"
 
@@ -110,6 +111,9 @@ namespace NetworkPhysicsCvars
 	int32 EnableDebugRPC = 1;
 #endif
 	FAutoConsoleVariableRef CVarEnableDebugRPC(TEXT("np2.EnableDebugRPC"), EnableDebugRPC, TEXT("Sends extra debug information to clients about server side input buffering"));
+
+	int32 NetworkPhysicsPredictionFrameOffset = 4;
+	FAutoConsoleVariableRef CVarNetworkPhysicsPredictionFrameOffset(TEXT("np2.NetworkPhysicsPredictionFrameOffset"), NetworkPhysicsPredictionFrameOffset, TEXT("Additional frame offset to be added to the local to server offset used by network prediction"));
 }
 
 const float RetryClientRestartThrottleTime = 0.5f;
@@ -196,7 +200,6 @@ APlayerController::APlayerController(const FObjectInitializer& ObjectInitializer
 	if (UPhysicsSettings::Get()->PhysicsPrediction.bEnablePhysicsPrediction)
 	{
 		bAsyncPhysicsTickEnabled = true;
-		AsyncPhysicsDataClass = UAsyncPhysicsData::StaticClass();
 	}
 #if UE_ENABLE_DEBUG_DRAWING
 	CurrentInputModeDebugString = TEXT("Default");
@@ -4794,6 +4797,8 @@ void APlayerController::GetLifetimeReplicatedProps(TArray< FLifetimeProperty > &
 {
 	Super::GetLifetimeReplicatedProps(OutLifetimeProps);
 
+	DISABLE_REPLICATED_PROPERTY(APlayerController, AsyncPhysicsDataComponent_DEPRECARED);
+
 	// These used to only replicate if PlayerCameraManager->GetViewTargetPawn() != GetPawn()
 	// But, since they also don't update unless that condition is true, these values won't change, thus won't send
 	// This is a little less efficient, but fits into the new condition system well, and shouldn't really add much overhead
@@ -4801,26 +4806,14 @@ void APlayerController::GetLifetimeReplicatedProps(TArray< FLifetimeProperty > &
 
 	// Replicate SpawnLocation for remote spectators
 	DOREPLIFETIME_CONDITION(APlayerController, SpawnLocation, COND_OwnerOnly);
-
-	DOREPLIFETIME_CONDITION(APlayerController, AsyncPhysicsDataComponent, COND_OwnerOnly);
 }
 
 void APlayerController::OnRep_AsyncPhysicsDataComponent()
-{
-	AsyncPhysicsDataComponent->SetDataClass(AsyncPhysicsDataClass);
-}
+{}
 
 void APlayerController::SetPlayer( UPlayer* InPlayer )
 {
 	FMoviePlayerProxyBlock MoviePlayerBlock;
-	if(AsyncPhysicsDataClass && GetLocalRole() == ROLE_Authority)
-	{
-		static const FName AsyncPhysicsDataComponentName(TEXT("PC_AsyncPhysicsDataComponent"));
-		AsyncPhysicsDataComponent = NewObject<UAsyncPhysicsInputComponent>(this, AsyncPhysicsDataComponentName);
-		AsyncPhysicsDataComponent->SetDataClass(AsyncPhysicsDataClass);
-		AsyncPhysicsDataComponent->RegisterComponent();
-	}
-
 	check(InPlayer!=NULL);
 
 	const bool bIsSameLevel = InPlayer->PlayerController && (InPlayer->PlayerController->GetLevel() == GetLevel());
@@ -6062,12 +6055,12 @@ void APlayerController::BeginReplication()
 
 UAsyncPhysicsData* APlayerController::GetAsyncPhysicsDataToWrite() const
 {
-	return AsyncPhysicsDataComponent ? AsyncPhysicsDataComponent->GetDataToWrite() : nullptr;
+	return nullptr;
 }
 
 const UAsyncPhysicsData* APlayerController::GetAsyncPhysicsDataToConsume() const
 {
-	return AsyncPhysicsDataComponent ? AsyncPhysicsDataComponent->GetDataToConsume() : nullptr;  
+	return nullptr;  
 }
 
 void APlayerController::ExecuteAsyncPhysicsCommand(const FAsyncPhysicsTimestamp& AsyncPhysicsTimestamp, UObject* OwningObject, const TFunction<void()>& Command, const bool bEnableResim)
@@ -6097,17 +6090,7 @@ FAsyncPhysicsTimestamp APlayerController::GetAsyncPhysicsTimestamp(float DeltaSe
 				const FReal DeltaTime = Solver->GetAsyncDeltaTime();
 				const int32 PendingSteps = (DeltaTime > 0.0) ? DeltaSeconds / DeltaTime : 0;
 
-				int32 LocalPhysicsStep;
-				if(Solver->IsGameThreadFrozen())
-				{
-					//We are calling inside the async tick, so use the current frame
-					LocalPhysicsStep = Solver->GetCurrentFrame();
-				}
-				else
-				{
-					//We are on the GT so give the upcoming async tick
-					LocalPhysicsStep = Solver->GetMarshallingManager().GetExternalTimestamp_External();
-				}
+				int32 LocalPhysicsStep = Solver->GetCurrentFrame();
 				
 				LocalPhysicsStep += PendingSteps;	//Add any pending steps user wants to wait on
 				Timestamp.ServerFrame = LocalPhysicsStep;
@@ -6127,6 +6110,16 @@ FAsyncPhysicsTimestamp APlayerController::GetAsyncPhysicsTimestamp(float DeltaSe
 
 void APlayerController::UpdateServerAsyncPhysicsTickOffset()
 {
+	if (UWorld* World = GetWorld())
+	{
+		if (FPhysScene* PhysScene = World->GetPhysicsScene())
+		{
+			if(PhysScene->GetSolver()->GetEvolution()->IsResimming())
+			{
+				return;
+			}
+		}
+	}
 	FAsyncPhysicsTimestamp Timestamp = GetAsyncPhysicsTimestamp();
 	if(ClientLatestAsyncPhysicsStepSent == Timestamp.LocalFrame)
 	{
@@ -6168,8 +6161,23 @@ void APlayerController::ClientCorrectionAsyncPhysicsTimestamp_Implementation(FAs
 		return;
 	}
 
-	const int32 NewOffset = Timestamp.ServerFrame - Timestamp.LocalFrame; //The new offset as reported by the server
-	ensureMsgf(NewOffset >= LocalToServerAsyncPhysicsTickOffset, TEXT("The offset between client and server can only ever increase"));
+	if (UWorld* World = GetWorld())
+	{
+		if (FPhysScene* PhysScene = World->GetPhysicsScene())
+		{
+			if (PhysScene->GetSolver()->GetEvolution()->IsResimming())
+			{
+				return;
+			}
+		}
+	}
+	FAsyncPhysicsTimestamp CurrentTimestamp = GetAsyncPhysicsTimestamp();
+	// We need to avoid changing this offset as much as possible since it will invalidate histories and will trigger resim
+	// To deal with that we compute a safe margin based on a user cvar + half the RTT
+	// This margin will only be applied the first time we will compute the offset 
+	const int32 FrameOffset = LocalToServerAsyncPhysicsTickOffset == 0 ? (NetworkPhysicsCvars::NetworkPhysicsPredictionFrameOffset + (CurrentTimestamp.LocalFrame - Timestamp.LocalFrame) / 2) : 0;
+
+	const int32 NewOffset = FMath::Max(LocalToServerAsyncPhysicsTickOffset,Timestamp.ServerFrame - Timestamp.LocalFrame + FrameOffset); //The new offset as reported by the server
 	LocalToServerAsyncPhysicsTickOffset = NewOffset;
 }
 
@@ -6206,19 +6214,9 @@ void APlayerController::AsyncPhysicsTickActor(float DeltaTime, float SimTime)
 
 		const FAsyncPhysicsTimestamp ActualTimestamp = GetAsyncPhysicsTimestamp();
 
-		//If the pending timestamp is bigger than the server frame we simple wait, otherwise make sure they aren't too early
-		if (ServerPendingTimestamps[0].ServerFrame <= ActualTimestamp.ServerFrame)
-		{
-			if (ServerPendingTimestamps[0].ServerFrame < ActualTimestamp.ServerFrame)
-			{
-				//The earliest pending client timestamp is too early, so their offset must be bigger than they thought
-				ServerLatestTimestampToCorrect.ServerFrame = ActualTimestamp.ServerFrame;
-				ServerLatestTimestampToCorrect.LocalFrame = ServerPendingTimestamps[0].LocalFrame;	//The client's local frame stays the same there's just a mismatch on the server frame
-			}
-
-			ServerPendingTimestamps.RemoveAt(0);	//We've updated the client if needed, so can discard
-			//NOTE: we purposely don't fixup any future pending client timestamps. This is because the RPC is unreliable so it's best to send the error correction redundantly over multiple frames if needed
-		}
+		ServerLatestTimestampToCorrect.ServerFrame = ActualTimestamp.ServerFrame;
+		ServerLatestTimestampToCorrect.LocalFrame = ServerPendingTimestamps[0].LocalFrame;
+		ServerPendingTimestamps.Reset();
 	}
 }
 
diff --git a/Engine/Source/Runtime/Engine/Public/Physics/Experimental/PhysScene_Chaos.h b/Engine/Source/Runtime/Engine/Public/Physics/Experimental/PhysScene_Chaos.h
--- a/Engine/Source/Runtime/Engine/Public/Physics/Experimental/PhysScene_Chaos.h
+++ b/Engine/Source/Runtime/Engine/Public/Physics/Experimental/PhysScene_Chaos.h
@@ -245,6 +245,9 @@ public:
 
 	UChaosEventRelay* GetChaosEventRelay() const { return ChaosEventRelay.Get(); }
 
+	/** Populate the replication cache from the list of registered components */
+	void PopulateReplicationCache(const int32 PhysicsStep);
+
 	// Storage structure for replication data
 	// probably should just expose read/write API not the structure directly itself like this.
 	struct FPrimitiveComponentReplicationCache
diff --git a/Engine/Source/Runtime/Engine/Public/Physics/NetworkPhysicsComponent.h b/Engine/Source/Runtime/Engine/Public/Physics/NetworkPhysicsComponent.h
new file mode 100644
--- /dev/null
+++ b/Engine/Source/Runtime/Engine/Public/Physics/NetworkPhysicsComponent.h
@@ -0,0 +1,434 @@
+// Copyright Epic Games, Inc. All Rights Reserved.
+
+#pragma once
+#include "RewindData.h"
+#include "Components/ActorComponent.h"
+#include "GameFramework/PawnMovementComponent.h"
+#include "GameFramework/PlayerController.h"
+#include "Engine/NetConnection.h"
+#include "Engine/World.h"
+#include "Physics/Experimental/PhysScene_Chaos.h"
+#include "Subsystems/WorldSubsystem.h"
+#include "PhysicsEngine/PhysicsSettings.h"
+
+#include "NetworkPhysicsComponent.generated.h"
+
+DECLARE_MULTICAST_DELEGATE_OneParam(FOnPreProcessInputsInternal, const int32);
+DECLARE_MULTICAST_DELEGATE_OneParam(FOnPostProcessInputsInternal, const int32);
+
+/** Templated datas history holding a datas buffer */
+template<typename DatasType>
+struct TNetRewindHistory : public Chaos::TDatasRewindHistory<DatasType>
+{
+	using Super = Chaos::TDatasRewindHistory<DatasType>;
+
+	FORCEINLINE TNetRewindHistory(const int32 FrameCount, const bool bIsHistoryLocal, UPackageMap* InPackageMap) :
+		Super(FrameCount, bIsHistoryLocal)
+	{
+		DatasWriter.SetAllowResize(true);
+		SetPackageMap(InPackageMap);
+	}
+
+	FORCEINLINE virtual ~TNetRewindHistory() {}
+
+	/** Set the package map for serialization */
+	FORCEINLINE virtual void SetPackageMap(class UPackageMap* InPackageMap) override
+	{
+		PackageMap = InPackageMap;
+
+		DatasWriter.PackageMap = PackageMap;
+		DatasReader.PackageMap = PackageMap;
+	}
+
+	/** Serialize the datas to the archive */
+	FORCEINLINE virtual void SerializeDatas(const uint32 StartFrame, const uint32 EndFrame, TArray<uint8>& ArchiveDatas, const int32 FrameOffset) const override
+	{
+		if (FrameOffset >= 0)
+		{
+			DatasWriter.Reset();
+
+			uint32 FramesCount = (uint32)Super::NumValidDatas(StartFrame, EndFrame);
+			DatasWriter << FramesCount;
+
+			DatasType FrameDatas;
+			for (uint32 FrameIndex = StartFrame; FrameIndex < EndFrame; ++FrameIndex)
+			{
+				const int32 LocalFrame = FrameIndex % Super::NumFrames;
+				if (FrameIndex == Super::DatasArray[LocalFrame].LocalFrame)
+				{
+					FrameDatas = Super::DatasArray[LocalFrame];
+					FrameDatas.ServerFrame = FrameDatas.LocalFrame + FrameOffset;
+					NetSerializeDatas(FrameDatas, DatasWriter);
+				}
+			}
+
+			ArchiveDatas = *DatasWriter.GetBuffer();
+		}
+	}
+
+	/** Deserialize the datas from the archive */
+	FORCEINLINE virtual void DeserializeDatas(const TArray<uint8>& ArchiveDatas, const int32 FrameOffset) override
+	{
+		if ((FrameOffset >= 0) && !ArchiveDatas.IsEmpty())
+		{
+			// TODO : temporary copy since the archive datas are const
+			TArray<uint8> LocalDatas = ArchiveDatas;
+			DatasReader.SetData(LocalDatas.GetData(), ArchiveDatas.Num() * 8);
+
+			uint32 FramesCount = 0;
+			DatasReader << FramesCount;
+
+			DatasType FrameDatas;
+			for (uint32 FrameIndex = 0; FrameIndex < FramesCount; ++FrameIndex)
+			{
+				NetSerializeDatas(FrameDatas, DatasReader);
+				FrameDatas.LocalFrame = FrameDatas.ServerFrame - FrameOffset;
+				if (FrameDatas.LocalFrame >= 0)
+				{
+					Super::RecordDatas(FrameDatas.LocalFrame, &FrameDatas);
+				}
+			}
+		}
+	}
+
+	/** Debug the datas from the archive */
+	FORCEINLINE virtual void DebugDatas(const TArray<uint8>& ArchiveDatas, TArray<int32>& LocalFrames, TArray<int32>& ServerFrames, TArray<int32>& InputFrames) override
+	{
+		if((PackageMap != nullptr) && !ArchiveDatas.IsEmpty())
+		{
+			// TODO : temporary copy since the archive datas are const
+			TArray<uint8> LocalDatas = ArchiveDatas;
+			DatasReader.SetData(LocalDatas.GetData(), ArchiveDatas.Num() * 8);
+
+			uint32 FramesCount = 0;
+			DatasReader << FramesCount;
+
+			LocalFrames.SetNum(FramesCount);
+			ServerFrames.SetNum(FramesCount);
+			InputFrames.SetNum(FramesCount);
+
+			DatasType FrameDatas;
+			for (uint32 FrameIndex = 0; FrameIndex < FramesCount; ++FrameIndex)
+			{
+				NetSerializeDatas(FrameDatas, DatasReader);
+				LocalFrames[FrameIndex] = FrameDatas.LocalFrame;
+				ServerFrames[FrameIndex] = FrameDatas.ServerFrame;
+				InputFrames[FrameIndex] = FrameDatas.InputFrame;
+			}
+		}
+	}
+
+private :
+
+	/** Use net serialize path to serialize datas  */
+	FORCEINLINE bool NetSerializeDatas(DatasType& FrameDatas, FBitArchive& Ar) const 
+	{
+		bool bOutSuccess = false;
+		UScriptStruct* ScriptStruct = DatasType::StaticStruct();
+		if (ScriptStruct->StructFlags & STRUCT_NetSerializeNative)
+		{
+			ScriptStruct->GetCppStructOps()->NetSerialize(Ar, PackageMap, bOutSuccess, &FrameDatas);
+		}
+		else
+		{
+			UE_LOG(LogPhysics, Error, TEXT("TNetRewindHistory::NetSerializeDatas called on data struct %s without a native NetSerialize"), *ScriptStruct->GetName());
+
+			// Not working for now since the packagemap could be null
+			// UNetConnection* Connection = CastChecked<UPackageMapClient>(PackageMap)->GetConnection();
+			// UNetDriver* NetDriver = Connection ? Connection->GetDriver() : nullptr;
+			// TSharedPtr<FRepLayout> RepLayout = NetDriver ? NetDriver->GetStructRepLayout(ScriptStruct) : nullptr;
+			//
+			// if (RepLayout.IsValid())
+			// {
+			// 	bool bHasUnmapped = false;
+			// 	RepLayout->SerializePropertiesForStruct(ScriptStruct, Ar, PackageMap, &FrameDatas, bHasUnmapped);
+			//
+			// 	bOutSuccess = true;
+			// }
+		}
+		return bOutSuccess;
+	}
+
+	// Datas bits writer to be used for serialization
+	mutable FNetBitWriter DatasWriter;
+
+	// Datas bits reader to be used for serialization
+	mutable FNetBitReader DatasReader;
+
+	// Package map used for the net bit writer/reader
+	UPackageMap* PackageMap = nullptr;
+};
+
+/**
+ * Network physics rewind callback to manage all the sim callbacks rewind functionalities
+ */
+struct FNetworkPhysicsCallback : public Chaos::IRewindCallback
+{
+	FNetworkPhysicsCallback(UWorld* InWorld) : World(InWorld) { }
+
+	// Delegate on the internal inputs process 
+	FOnPreProcessInputsInternal PreProcessInputsInternal;
+	FOnPostProcessInputsInternal PostProcessInputsInternal;
+
+	// Rewind API
+	virtual void ProcessInputs_External(int32 PhysicsStep, const TArray<Chaos::FSimCallbackInputAndObject>& SimCallbackInputs);
+	virtual void ProcessInputs_Internal(int32 PhysicsStep, const TArray<Chaos::FSimCallbackInputAndObject>& SimCallbackInputs) override;
+	virtual void ApplyCallbacks_Internal(int32 PhysicsStep, const TArray<Chaos::ISimCallbackObject*>& SimCallbackObjects) override;
+	virtual int32 TriggerRewindIfNeeded_Internal(int32 LatestStepCompleted) override;
+	virtual void RegisterRewindableSimCallback_Internal(Chaos::ISimCallbackObject* SimCallbackObject) override
+	{
+		if (SimCallbackObject && SimCallbackObject->HasOption(Chaos::ESimCallbackOptions::Rewind))
+		{
+			RewindableCallbackObjects.Add(SimCallbackObject);
+		}
+	}
+
+	// Updates the TMap on PhysScene that stores (non interpolated) physics data for replication.
+	// 
+	// Needs to be called from PT context to access fixed tick handle
+	// but also needs to be able to access GT data (actor iterator, actor state)
+	void UpdateReplicationMap_Internal(int32 PhysicsStep);
+
+	// Update client player on GT
+	void UpdateClientPlayer_External(int32 PhysicsStep);
+
+	// Update server player on GT
+	void UpdateServerPlayer_External(int32 PhysicsStep);
+
+	// World owning that callback
+	UWorld* World = nullptr;
+
+	// List of rewindable sim callback objects
+	TArray<Chaos::ISimCallbackObject*> RewindableCallbackObjects;
+};
+
+/**
+ * Network physics manager to initialize datas required for rewind/resim
+ */
+UCLASS()
+class ENGINE_API UNetworkPhysicsSystem : public UWorldSubsystem
+{
+public:
+
+	GENERATED_BODY()
+	UNetworkPhysicsSystem();
+
+	friend struct FNetworkPhysicsCallback;
+
+	// Subsystem Init/Deinit
+	virtual void Initialize(FSubsystemCollectionBase& Collection) override;
+	virtual void Deinitialize() override;
+
+	// Delegate at world init 
+	void OnWorldPostInit(UWorld* World, const UWorld::InitializationValues);
+
+	// Register a component to be used by the callback 
+	void RegisterNetworkComponent(class UNetworkPhysicsComponent* NetworkComponent) {NetworkComponents.Add(NetworkComponent); }
+
+	// Remove a network component from registered list
+	void UnregisterNetworkComponent(class UNetworkPhysicsComponent* NetworkComponent) { NetworkComponents.Remove(NetworkComponent); }
+
+private:
+
+	// List of physics network components that will be used by the rewind callback
+	TArray<class UNetworkPhysicsComponent*> NetworkComponents;
+};
+
+/**
+ * Base network physics datas that will be used by physics
+ */
+ USTRUCT()
+struct FNetworkPhysicsDatas
+{
+	GENERATED_USTRUCT_BODY()
+
+	virtual ~FNetworkPhysicsDatas() = default;
+
+	// Server frame at which this datas has been generated
+	UPROPERTY()
+	int32 ServerFrame = INDEX_NONE;
+
+	// Local frame at which this datas has been generated
+	UPROPERTY()
+	int32 LocalFrame = INDEX_NONE;
+
+	// Input frame used to generate the network datas
+	UPROPERTY()
+	int32 InputFrame = INDEX_NONE;
+
+	// Serialize the datas into/from the archive
+	void SerializeFrames(FArchive& Ar)
+	{
+		Ar << ServerFrame;
+		Ar << LocalFrame;
+		Ar << InputFrame;
+	}
+
+	// Apply the datas from onto the network physics component
+	virtual void ApplyDatas(UActorComponent* NetworkComponent) const {}
+
+	// Build the datas from the network physics component
+	virtual void BuildDatas(const UActorComponent* NetworkComponent) {}
+
+	friend UNetworkPhysicsComponent;
+};
+
+/**
+ * Network physics component that will be attached to any player controller
+ */
+UCLASS(BlueprintType)
+class ENGINE_API UNetworkPhysicsComponent : public UActorComponent
+{
+	GENERATED_UCLASS_BODY()
+public:
+	UNetworkPhysicsComponent();
+
+	// Get the player controller on which the component is attached
+	virtual APlayerController* GetPlayerController() const;
+
+	// Init the network physics component 
+	void InitPhysics();
+
+	// Server RPC to receive inputs from client
+	UFUNCTION(Server, unreliable, WithValidation)
+	void ServerReceiveInputsDatas(const TArray<uint8>& ClientInputs);
+
+	// Async physics tick component function per frame from the solver
+	virtual void AsyncPhysicsTickComponent(float DeltaTime, float SimTime) override;
+
+	// Function to init the replicated properties
+	virtual void GetLifetimeReplicatedProps(TArray<FLifetimeProperty>& OutLifeTimeProps) const override;
+
+	// Send the inputs replicated datas
+	void SendLocalInputsDatas();
+
+	// Send the states replicated datas
+	void SendLocalStatesDatas();
+
+	// Delegate linked to the physics rewind callback to send record local inputs/states
+	void OnPreProcessInputsInternal(const int32 PhysicsStep);
+
+	// Delegate linked to the physics rewind callback to send record local inputs/states
+	void OnPostProcessInputsInternal(const int32 PhysicsStep);
+
+	// Correct the player controller Server to local offset based on the received replicated states
+	void CorrectServerToLocalOffset(const int32 LocalToServerOffset);
+
+	// Used to create any physics engine information for this component 
+	virtual void BeginPlay() override;
+
+	// Register the component into the network manager
+	virtual void InitializeComponent() override;
+
+	// Unregister the component from the network manager
+	virtual void UninitializeComponent() override;
+
+	// Register and create the states/inputs history
+	template<typename PhysicsTraits>
+	void CreateDatasHistory(UActorComponent* HistoryComponent);
+
+	// Remove states/inputs history from rewind datas
+	void RemoveDatasHistory();
+
+	// Add states/inputs history to rewind datas
+	void AddDatasHistory();
+
+	// Get the datas factory that will be used for net serialization
+	TSharedPtr<Chaos::FBaseRewindHistory>& GetStatesHistory() { return StatesHistory; }
+
+	// Get the datas factory that will be used for net serialization
+	TSharedPtr<Chaos::FBaseRewindHistory>& GetInputsHistory() { return InputsHistory; }
+
+	// Check if the world is on server
+	bool HasServerWorld() const;
+
+	// Check if the player controller exists and is local
+	bool HasLocalController() const;
+
+protected : 
+
+	// Update the histories packagemap for serialization 
+	void UpdatePackageMap();
+
+	// repnotify for the inputs on the client
+	UFUNCTION()
+	void OnRep_SetReplicatedInputs();
+
+	// repnotify for the states on the client
+	UFUNCTION()
+	void OnRep_SetReplicatedStates();
+
+	// replicated physics inputs
+	UPROPERTY(Transient, ReplicatedUsing = OnRep_SetReplicatedInputs)
+	TArray<uint8> ReplicatedInputs;
+
+	// replicated physics states 
+	UPROPERTY(Transient, ReplicatedUsing = OnRep_SetReplicatedStates)
+	TArray<uint8> ReplicatedStates;
+
+	// List of inputs that are processed locally
+	TArray<uint8> LocalInputs;
+
+	// List of states that are processed locally
+	TArray<uint8> LocalStates;
+
+	// Frame counter to compute the local to server offset
+	int32 FrameCounter = 0;
+
+private:
+
+	friend FNetworkPhysicsCallback;
+
+	// States history uses to rewind simulation 
+	TSharedPtr<Chaos::FBaseRewindHistory> StatesHistory;
+
+	// Inputs history used during simulation
+	TSharedPtr<Chaos::FBaseRewindHistory> InputsHistory;
+
+	// Local temporary inputs datas used by pre/post process inputs functions
+	TUniquePtr<FNetworkPhysicsDatas> InputsDatas;
+
+	// Local temporary states datas used by pre/post process inputs functions
+	TUniquePtr<FNetworkPhysicsDatas> StatesDatas;
+
+	// Specify how much times the network will resend the inputs in case of packet loss
+	int8 InputsRedundancy = 4;
+
+	// Current index used in the inputs offsets
+	int8 InputsIndex = 0;
+
+	// Inputs offsets defined on PT based on the newly recorded inputs datas
+	TArray<int32> InputsOffsets;
+
+	// Specify how much times the network will resend the states in case of packet loss
+	int8 StatesRedundancy = 1;
+
+	// Current index used in the states offsets
+	int8 StatesIndex = 0;
+
+	// States offsets defined on PT based on the newly recorded states datas
+	TArray<int32> StatesOffsets;
+
+	// Actor component that will be used to fill the histories
+	TObjectPtr<UActorComponent> ActorComponent;
+};
+
+template<typename PhysicsTraits>
+FORCEINLINE void UNetworkPhysicsComponent::CreateDatasHistory(UActorComponent* HistoryComponent)
+{
+	APlayerController* Controller = GetPlayerController();
+	const bool bIsLocalHistory = (Controller && Controller->IsLocalController());
+	const int32 NumFrames = UPhysicsSettings::Get()->GetPhysicsHistoryCount();
+	
+	UPackageMap* PackageMap = (Controller  && Controller->GetNetConnection())? Controller->GetNetConnection()->PackageMap : nullptr;
+
+	InputsHistory = MakeShared<TNetRewindHistory<typename PhysicsTraits::InputsType>>(NumFrames, bIsLocalHistory, PackageMap);
+	StatesHistory = MakeShared<TNetRewindHistory<typename PhysicsTraits::StatesType>>(NumFrames, bIsLocalHistory, PackageMap);
+
+	InputsDatas = MakeUnique<typename PhysicsTraits::InputsType>();
+	StatesDatas = MakeUnique<typename PhysicsTraits::StatesType>();
+
+	ActorComponent = HistoryComponent;
+	
+	AddDatasHistory();
+}
diff --git a/Engine/Source/Runtime/Experimental/Chaos/Private/PBDRigidsSolver.cpp b/Engine/Source/Runtime/Experimental/Chaos/Private/PBDRigidsSolver.cpp
--- a/Engine/Source/Runtime/Experimental/Chaos/Private/PBDRigidsSolver.cpp
+++ b/Engine/Source/Runtime/Experimental/Chaos/Private/PBDRigidsSolver.cpp
@@ -638,6 +638,18 @@ namespace Chaos
 	int32 LogCorruptMap = 0;
 	FAutoConsoleVariableRef CVarLogCorruptMap(TEXT("p.LogCorruptMap"), LogCorruptMap, TEXT(""));
 
+	void FPBDRigidsSolver::ApplyCallbacks_Internal()
+	{
+		Super::ApplyCallbacks_Internal();
+
+		if (IRewindCallback* RewindCallback = GetRewindCallback())
+		{
+			if(bGameThreadFrozen)
+			{ 
+				RewindCallback->ApplyCallbacks_Internal(GetCurrentFrame(), SimCallbackObjects);
+			}
+		}
+	}
 
 	void FPBDRigidsSolver::UnregisterObject(FSingleParticlePhysicsProxy* Proxy)
 	{
@@ -965,6 +977,10 @@ namespace Chaos
 		MRewindCallback = MoveTemp(RewindCallback);
 		MarshallingManager.SetHistoryLength_Internal(NumFrames);
 		MEvolution->SetRewindData(GetRewindData());
+		if(MRewindCallback)
+		{
+			MRewindCallback->RewindData = MRewindData.Get();
+		}
 		
 		UpdateIsDeterministic();
 	}
@@ -1603,6 +1619,14 @@ namespace Chaos
 			{
 				UnregistrationWatchers.Add(SimCallbackObject);
 			}
+
+			if (SimCallbackObject->HasOption(ESimCallbackOptions::Rewind))
+			{
+				if (MRewindCallback.IsValid())
+				{
+					MRewindCallback->RegisterRewindableSimCallback_Internal(SimCallbackObject);
+				}
+			}
 		}
 
 		//save any pending data for this particular interval
@@ -1670,12 +1694,12 @@ namespace Chaos
 #if !(UE_BUILD_SHIPPING || UE_BUILD_TEST)
 			if (bEnableNetworkPredictionDebug)
 			{
-				UE_LOG(LogChaos, Log, TEXT("	-> Trying to Rewind Frame = %d | At Time = %f | Num Steps = %d | Resim Step = %d | Last Step = %d"), CurrentFrame,  MTime, NumResimSteps, ResimStep, LastStep);
+				UE_LOG(LogChaos, Log, TEXT("	-> Trying to Rewind Frame = %d | At Time = %f | Num Steps = %d | Resim Step = %d | Last Step = %d | Manager Size = %d"), CurrentFrame,  MTime, NumResimSteps, ResimStep, LastStep, MarshallingManager.GetNumHistory_Internal());
 				for(auto& Handle : GetParticles().GetNonDisabledDynamicView())
 				{
 					if(!Handle.IsSleeping())
 					{
-						UE_LOG(LogChaos, Log, TEXT("Particle Active At Position = %s | Velocity = %s | Quaternion = %s | Omega = %s"),
+						UE_LOG(LogChaos, Log, TEXT("Particle Dynamic At Position = %s | Velocity = %s | Quaternion = %s | Omega = %s"),
 							*Handle.X().ToString(), *Handle.V().ToString(), *Handle.R().ToString(), *Handle.W().ToString());
 					}
 					else
@@ -1722,12 +1746,22 @@ namespace Chaos
 								UE_LOG(LogChaos, Log, TEXT("		-> Re-simulating Frame = %d "), Step);
 								for(auto& Handle : GetParticles().GetNonDisabledDynamicView())
 								{
-									if(!Handle.IsSleeping()) 
+									if (GetEvolution()->GetIslandManager().GetParticleIsland(Handle.Handle()))
 									{
-										UE_LOG(LogChaos, Log, TEXT("Particle Dynamic At Position = %s | Velocity = %s | Quaternion = %s | Omega = %s | Resim Frame = %d | Sync State = %d | Needs Resim = %d"),   
-											*Handle.X().ToString(), *Handle.V().ToString(), *Handle.R().ToString(), *Handle.W().ToString(), 
-											GetEvolution()->GetIslandManager().GetParticleIsland(Handle.Handle())->GetResimFrame(), (uint8)Handle.SyncState(),
-											GetEvolution()->GetIslandManager().GetParticleIsland(Handle.Handle())->NeedsResim());
+										if (!Handle.IsSleeping())
+										{
+											UE_LOG(LogChaos, Log, TEXT("Particle Dynamic At Position = %s | Velocity = %s | Quaternion = %s | Omega = %s | Resim Frame = %d | Sync State = %d | Needs Resim = %d | Unique Idx = %d"),
+												*Handle.X().ToString(), *Handle.V().ToString(), *Handle.R().ToString(), *Handle.W().ToString(),
+												GetEvolution()->GetIslandManager().GetParticleIsland(Handle.Handle())->GetResimFrame(), (uint8)Handle.SyncState(),
+												GetEvolution()->GetIslandManager().GetParticleIsland(Handle.Handle())->NeedsResim(), Handle.UniqueIdx().Idx);
+										}
+										else
+										{
+											UE_LOG(LogChaos, Log, TEXT("Particle Sleeping At Position = %s | Velocity = %s | Quaternion = %s | Omega = %s | Resim Frame = %d | Sync State = %d | Needs Resim = %d | Unique Idx = %d"),
+												*Handle.X().ToString(), *Handle.V().ToString(), *Handle.R().ToString(), *Handle.W().ToString(),
+												GetEvolution()->GetIslandManager().GetParticleIsland(Handle.Handle())->GetResimFrame(), (uint8)Handle.SyncState(),
+												GetEvolution()->GetIslandManager().GetParticleIsland(Handle.Handle())->NeedsResim(), Handle.UniqueIdx().Idx);
+										}
 									}
 								}
 							}
@@ -2261,6 +2295,7 @@ TRACE_COUNTER_SET(ChaosTraceCounter_##Name, Value)
 	void FPBDRigidsSolver::FinalizeRewindData(const TParticleView<TPBDRigidParticles<FReal,3>>& DirtyParticles)
 	{
 		using namespace Chaos;
+
 		//Simulated objects must have their properties captured for rewind
 		if(MRewindData && DirtyParticles.Num())
 		{
diff --git a/Engine/Source/Runtime/Experimental/Chaos/Private/RewindData.cpp b/Engine/Source/Runtime/Experimental/Chaos/Private/RewindData.cpp
--- a/Engine/Source/Runtime/Experimental/Chaos/Private/RewindData.cpp
+++ b/Engine/Source/Runtime/Experimental/Chaos/Private/RewindData.cpp
@@ -74,8 +74,7 @@ bool FGeometryParticleStateBase::IsResimFrameValid(const FGeometryParticleHandle
 		if (RigidSolver->GetEvolution() && RigidSolver->GetEvolution()->IsResimming())
 		{
 			const int32 ResimFrame = RigidSolver->GetEvolution()->GetIslandManager().GetParticleResimFrame(&Handle);
-
-			if (ResimFrame != INDEX_NONE && ResimFrame <= FrameAndPhase.Frame)
+			if (ResimFrame != INDEX_NONE /*&& ResimFrame <= FrameAndPhase.Frame*/)
 			{
 				return false;
 			}
@@ -87,11 +86,13 @@ bool FGeometryParticleStateBase::IsResimFrameValid(const FGeometryParticleHandle
 template <bool bSkipDynamics>
 bool FGeometryParticleStateBase::IsInSync(const FGeometryParticleHandle& Handle, const FFrameAndPhase FrameAndPhase, const FDirtyPropertiesPool& Pool) const
 {
-	if(Chaos::FPhysicsSolverBase::IsNetworkPhysicsPredictionEnabled() && !TargetPositions.IsEmpty() && !TargetVelocities.IsEmpty())
-	{ 
-		return IsResimFrameValid(Handle, FrameAndPhase);
+	if(Chaos::FPhysicsSolverBase::IsNetworkPhysicsPredictionEnabled() && !TargetPositions.IsEmpty() && !TargetVelocities.IsEmpty() && !TargetStates.IsEmpty())
+	{
+		if(!IsResimFrameValid(Handle, FrameAndPhase))
+		{
+			return false;
+		}
 	}
-	else
 	{ 
 		if(!ParticlePositionRotation.IsInSync(Handle, FrameAndPhase, Pool))
 		{
@@ -162,7 +163,7 @@ bool FJointStateBase::IsInSync(const FPBDJointConstraintHandle& Handle, const FF
 
 void FRewindData::ApplyInputs(const int32 ApplyFrame, const bool bResetSolver)
 {
-	for (TWeakPtr<FBaseInputsHistory>& InputsHistory : InputsHistories)
+	for (TWeakPtr<FBaseRewindHistory>& InputsHistory : InputsHistories)
 	{
 		if (InputsHistory.IsValid())
 		{
@@ -172,7 +173,7 @@ void FRewindData::ApplyInputs(const int32 ApplyFrame, const bool bResetSolver)
 }
 void FRewindData::RewindStates(const int32 RewindFrame, const bool bResetSolver)
 {
-	for (TWeakPtr<FBaseStatesHistory>& StatesHistory : StatesHistories)
+	for (TWeakPtr<FBaseRewindHistory>& StatesHistory : StatesHistories)
 	{
 		if (StatesHistory.IsValid())
 		{
@@ -206,6 +207,45 @@ void FRewindData::ApplyTargets(const int32 Frame, const bool bResetSimulation)
 		const bool bResimAsFollower = DirtyParticleInfo.bResimAsFollower;
 		RewindHelper(PTParticle, bResimAsFollower, History.TargetPositions, [](auto Particle, const auto& Data) {Particle->SetXR(Data); });
 		RewindHelper(PTParticle->CastToKinematicParticle(), bResimAsFollower, History.TargetVelocities, [](auto Particle, const auto& Data) {Particle->SetVelocities(Data); });
+		RewindHelper(PTParticle->CastToRigidParticle(), bResimAsFollower, History.TargetStates, [this](auto Particle, const auto& Data) 
+		{ 
+			if (Particle == nullptr || Solver->GetEvolution() == nullptr)
+			{
+				return;
+			}
+
+			// Enable or disable the particle
+			if (Particle->Disabled() != Data.Disabled())
+			{
+				if (Data.Disabled())
+				{
+					Solver->GetEvolution()->DisableParticle(Particle);
+				}
+				else
+				{
+					Solver->GetEvolution()->EnableParticle(Particle);
+				}
+			}
+
+			// If we changed kinematics we need to rebuild the inertia conditioning
+			const bool bDirtyInertiaConditioning = (Particle->ObjectState() != Data.ObjectState());
+			if (bDirtyInertiaConditioning)
+			{
+				Particle->SetInertiaConditioningDirty();
+			}
+		
+			Particle->SetDisabled(Data.Disabled());
+			Solver->GetEvolution()->SetParticleObjectState(Particle, Data.ObjectState());
+
+			if(Data.ObjectState() == EObjectStateType::Dynamic)
+			{
+				Particle->SetResimType(EResimType::FullResim);
+			}
+			else if((Data.ObjectState() == EObjectStateType::Static) || (Data.ObjectState() == EObjectStateType::Kinematic))
+			{
+				Particle->SetResimType(EResimType::ResimAsFollower);
+			}
+		});
 #if !(UE_BUILD_SHIPPING || UE_BUILD_TEST)
 		if (!History.TargetPositions.IsClean(RewindFrameAndPhase) && Chaos::FPhysicsSolverBase::CanDebugNetworkPhysicsPrediction())
 		{
@@ -222,7 +262,7 @@ bool FRewindData::RewindToFrame(int32 Frame)
 	EnsureIsInPhysicsThreadContext();
 	//Can't go too far back
 	const int32 EarliestFrame = GetEarliestFrame_Internal();
-	if(Frame < EarliestFrame)
+	if (Frame < EarliestFrame)
 	{
 		return false;
 	}
@@ -477,7 +517,7 @@ void FRewindData::AdvanceFrameImp(IResimCacheBase* ResimCache)
 			auto& Info = DirtyObjects.GetDenseAt(DirtyIdx);
 
 			ensure(IsResimAndInSync(*Info.GetObjectPtr()) || Info.GetHistory().IsClean(FrameAndPhase));  //Sim hasn't run yet so PostCallbacks (sim results) should be clean
-
+	
 			//if hasn't changed in a while stop tracking
 			if (Info.LastDirtyFrame < EarliestFrame)
 			{
@@ -787,6 +827,71 @@ void FRewindData::ExtendHistoryWithFrame(const int32 Frame)
 	FramesSaved = FMath::Max(CurFrame - Frame+1, FramesSaved);
 }
 
+int32 FRewindData::FindValidResimFrame(const int32 ResimFrame)
+{
+	int32 ValidFrame = INDEX_NONE;
+	if(ResimFrame > 0)
+	{
+		EnsureIsInPhysicsThreadContext();
+
+		// First frame of the history datas
+		const int32 EarliestFrame = FMath::Max(GetEarliestFrame_Internal(), 0);
+	
+		for(ValidFrame = ResimFrame; ValidFrame >= EarliestFrame; --ValidFrame)
+		{
+			bool bHasTargetHistory = true;
+			for (FDirtyParticleInfo& DirtyParticleInfo : DirtyParticles)
+			{
+				FGeometryParticleHandle* PTParticle = DirtyParticleInfo.GetObjectPtr();
+				FGeometryParticleStateBase& History = DirtyParticleInfo.GetHistory();
+
+				const bool bResimAsFollower = DirtyParticleInfo.bResimAsFollower;
+
+				const FFrameAndPhase FrameAndPhase{ ValidFrame, FFrameAndPhase::PostPushData };
+				if(const FParticleDynamicMisc* DynamicMisc = History.DynamicsMisc.Read(FrameAndPhase, PropertiesPool))
+				{
+					if (!DynamicMisc->Disabled() && (DynamicMisc->ObjectState() == EObjectStateType::Dynamic) && !History.TargetPositions.IsEmpty() && !History.TargetVelocities.IsEmpty() && !History.TargetStates.IsEmpty())
+					{
+						if (bResimAsFollower || History.TargetPositions.IsClean(FrameAndPhase) || History.TargetVelocities.IsClean(FrameAndPhase) || History.TargetStates.IsClean(FrameAndPhase))
+						{
+							bHasTargetHistory = false;
+							break;
+						}
+					}
+				}
+			}
+			if(bHasTargetHistory)
+			{
+				for(auto& InputsHistory : InputsHistories)
+				{
+					if(!InputsHistory.Pin().Get()->HasValidDatas(ValidFrame))
+					{
+						bHasTargetHistory = false;
+						break;
+					}
+				}
+			}
+			if(bHasTargetHistory)
+			{
+				for (auto& StatesHistory : StatesHistories)
+				{
+					if (!StatesHistory.Pin().Get()->HasValidDatas(ValidFrame))
+					{
+						bHasTargetHistory = false;
+						break;
+					}
+				}
+			}
+
+			if(bHasTargetHistory)
+			{
+				break;
+			}
+		}
+	}
+	return ValidFrame;
+}
+
 
 void FRewindData::PushStateAtFrame(FGeometryParticleHandle& Handle, int32 Frame, FFrameAndPhase::EParticleHistoryPhase Phase, 
 	const FVector& Position, const FQuat& Quaternion, const FVector& LinVelocity, const FVector& AngVelocity, const bool bShouldSleep)
@@ -794,32 +899,29 @@ void FRewindData::PushStateAtFrame(FGeometryParticleHandle& Handle, int32 Frame,
 	const bool bRecordingHistory = !IsResimAndInSync(Handle);
 
 	FDirtyParticleInfo& Info = FindOrAddDirtyObj(Handle);
-	const int32 LastFrame = Info.LastDirtyFrame;
-	FGeometryParticleStateBase& Latest = Info.AddFrame(Frame);
+	FGeometryParticleStateBase& Latest = Info.GetHistory();
 	const FFrameAndPhase FrameAndPhase{ Frame, Phase };
 
 	if (bRecordingHistory || Latest.TargetPositions.IsClean(FrameAndPhase))
 	{
-		FParticlePositionRotation& PositionRotation = Latest.TargetPositions.WriteAccessMonotonic(FrameAndPhase, PropertiesPool);
+		FParticlePositionRotation& PositionRotation = Latest.TargetPositions.Insert(FrameAndPhase, PropertiesPool);
 		PositionRotation.SetX(Position);
 		PositionRotation.SetR(Quaternion);
 	}
 
 	if (bRecordingHistory || Latest.TargetVelocities.IsClean(FrameAndPhase))
 	{
-		FParticleVelocities& PreVelocities = Latest.TargetVelocities.WriteAccessMonotonic(FrameAndPhase, PropertiesPool);
+		FParticleVelocities& PreVelocities = Latest.TargetVelocities.Insert(FrameAndPhase, PropertiesPool);
 		PreVelocities.SetV(LinVelocity);
 		PreVelocities.SetW(AngVelocity);
 	}
-	if (bRecordingHistory || Latest.DynamicsMisc.IsClean(FrameAndPhase))
+
+	if (bRecordingHistory || Latest.TargetStates.IsClean(FrameAndPhase))
 	{
-		if (bShouldSleep)
-		{
-			FParticleDynamicMisc& PreDynamicsMisc = Latest.DynamicsMisc.WriteAccessMonotonic(FrameAndPhase, PropertiesPool);
-			PreDynamicsMisc.SetObjectState(EObjectStateType::Sleeping);
-		}
+		FParticleDynamicMisc& PreDynamicsMisc = Latest.TargetStates.Insert(FrameAndPhase, PropertiesPool);
+		PreDynamicsMisc.SetObjectState(bShouldSleep ? EObjectStateType::Sleeping : EObjectStateType::Dynamic);
+		PreDynamicsMisc.SetDisabled(false);
 	}
-	Info.LastDirtyFrame = LastFrame;
 }
 
 void FRewindData::PushPTDirtyData(TPBDRigidParticleHandle<FReal,3>& Handle,const int32 SrcDataIdx)
@@ -828,6 +930,8 @@ void FRewindData::PushPTDirtyData(TPBDRigidParticleHandle<FReal,3>& Handle,const
 
 	FDirtyParticleInfo& Info = FindOrAddDirtyObj(Handle);
 	FGeometryParticleStateBase& Latest = Info.AddFrame(CurFrame);
+
+
 	const FFrameAndPhase FrameAndPhase{ CurFrame, FFrameAndPhase::PostCallbacks };
 
 	ensure(!bRecordingHistory || Latest.IsCleanExcludingDynamics(FrameAndPhase));	//PostCallbacks should be clean before we write sim results
@@ -865,21 +969,8 @@ FJointState FRewindData::GetPastJointStateAtFrame(const FPBDJointConstraintHandl
 void FRewindData::SetTargetStateAtFrame(FGeometryParticleHandle& Handle, const int32 Frame, FFrameAndPhase::EParticleHistoryPhase Phase,
 	const FVector& Position, const FQuat& Quaternion, const FVector& LinVelocity, const FVector& AngVelocity, const bool bShouldSleep)
 {
-	int32 PositionValidCount = INDEX_NONE;
-	int32 VelocityValidCount = INDEX_NONE;
-	int32 PositionNextIterator = INDEX_NONE;
-	int32 VelocityNextIterator = INDEX_NONE;
-
-	ExtractHistoryState(Handle,
-		PositionValidCount, VelocityValidCount, PositionNextIterator, VelocityNextIterator);
-
-	ClearPhaseAndFuture(Handle, Frame, Phase);
-
 	PushStateAtFrame(Handle, Frame, Phase,
 		Position, Quaternion, LinVelocity, AngVelocity, bShouldSleep);
-
-	RestoreHistoryState(Handle,
-		PositionValidCount, VelocityValidCount, PositionNextIterator, VelocityNextIterator);
 }
 
 }
diff --git a/Engine/Source/Runtime/Experimental/Chaos/Public/Chaos/Framework/PhysicsSolverBase.h b/Engine/Source/Runtime/Experimental/Chaos/Public/Chaos/Framework/PhysicsSolverBase.h
--- a/Engine/Source/Runtime/Experimental/Chaos/Public/Chaos/Framework/PhysicsSolverBase.h
+++ b/Engine/Source/Runtime/Experimental/Chaos/Public/Chaos/Framework/PhysicsSolverBase.h
@@ -338,10 +338,6 @@ namespace Chaos
 			}
 
 			RegisterSimCallbackObject_External(NewCallbackObject);
-			if (NewCallbackObject->HasOption(ESimCallbackOptions::Rewind))
-			{
-				EnqueueSimcallbackRewindRegisteration(NewCallbackObject);
-			}
 			return NewCallbackObject;
 		}
 
@@ -515,7 +511,7 @@ namespace Chaos
 			bGameThreadFrozen = InGameThreadFrozen;
 		}
 
-		void ApplyCallbacks_Internal()
+		virtual void ApplyCallbacks_Internal()
 		{
 			QUICK_SCOPE_CYCLE_COUNTER(ApplySimCallbacks);
 
diff --git a/Engine/Source/Runtime/Experimental/Chaos/Public/Chaos/PBDRigidsSOAs.h b/Engine/Source/Runtime/Experimental/Chaos/Public/Chaos/PBDRigidsSOAs.h
--- a/Engine/Source/Runtime/Experimental/Chaos/Public/Chaos/PBDRigidsSOAs.h
+++ b/Engine/Source/Runtime/Experimental/Chaos/Public/Chaos/PBDRigidsSOAs.h
@@ -344,7 +344,7 @@ public:
 		if(FPBDRigidParticleHandle* Rigid =  Particle->CastToRigidParticle())
 		{
 			//if active it's already in the dirty view
-			if (!ActiveParticlesMapArray.Contains(Rigid))
+			if (!ActiveParticlesMapArray.Contains(Rigid) && !MovingKinematicsMapArray.Contains(Rigid))
 			{
 				TransientDirtyMapArray.Insert(Rigid);
 				if (bUpdateViews)
@@ -961,8 +961,11 @@ private:
 
 		if(bStillDirty)
 		{
-			//no longer active, but still dirty
-			TransientDirtyMapArray.Insert(Particle);
+			if(!MovingKinematicsMapArray.Contains(Particle))
+			{ 
+				//no longer active, but still dirty
+				TransientDirtyMapArray.Insert(Particle);
+			}
 		}
 		else
 		{
diff --git a/Engine/Source/Runtime/Experimental/Chaos/Public/Chaos/SimCallbackObject.h b/Engine/Source/Runtime/Experimental/Chaos/Public/Chaos/SimCallbackObject.h
--- a/Engine/Source/Runtime/Experimental/Chaos/Public/Chaos/SimCallbackObject.h
+++ b/Engine/Source/Runtime/Experimental/Chaos/Public/Chaos/SimCallbackObject.h
@@ -154,12 +154,15 @@ public:
 	}
 
 	// Rewind API
+	virtual void InjectInputs_External(int32 PhysicsStep, int32 NumSteps) {}
+	virtual void ProcessInputs_Internal(int32 PhysicsStep) {}
+	virtual void ProcessInputs_External(int32 PhysicsStep) {}
+	
 	virtual int32 TriggerRewindIfNeeded_Internal(int32 LastCompletedStep)
-	{ 
-		ensure(false); 
-		return INDEX_NONE; 
+	{
+		return INDEX_NONE;
 	}
-
+	
 	virtual void ApplyCorrections_Internal(int32 PhysicsStep, FSimCallbackInput* Input)
 	{
 		ensure(false);
@@ -178,7 +181,7 @@ public:
 	{
 		return HasOption(ESimCallbackOptions::RunOnFrozenGameThread);
 	}
-	
+
 protected:
 
 	ISimCallbackObject(const ESimCallbackOptions InOptions = ESimCallbackOptions::Presimulate)
diff --git a/Engine/Source/Runtime/Experimental/Chaos/Public/PBDRigidsSolver.h b/Engine/Source/Runtime/Experimental/Chaos/Public/PBDRigidsSolver.h
--- a/Engine/Source/Runtime/Experimental/Chaos/Public/PBDRigidsSolver.h
+++ b/Engine/Source/Runtime/Experimental/Chaos/Public/PBDRigidsSolver.h
@@ -327,6 +327,8 @@ namespace Chaos
 		// Interop utilities
 		void SetParticleDynamicMisc(FPBDRigidParticleHandle* Rigid, const FParticleDynamicMisc& DynamicMisc);
 
+		// Apply callbacks internally 
+		virtual void ApplyCallbacks_Internal() override;
 
 	private:
 
diff --git a/Engine/Source/Runtime/Experimental/Chaos/Public/RewindData.h b/Engine/Source/Runtime/Experimental/Chaos/Public/RewindData.h
--- a/Engine/Source/Runtime/Experimental/Chaos/Public/RewindData.h
+++ b/Engine/Source/Runtime/Experimental/Chaos/Public/RewindData.h
@@ -4,6 +4,7 @@
 #include "Chaos/ParticleHandle.h"
 #include "PhysicsProxy/SingleParticlePhysicsProxyFwd.h"
 #include "Chaos/Framework/PhysicsSolverBase.h"
+#include "Serialization/BufferArchive.h"
 #include "Containers/CircularBuffer.h"
 #include "Chaos/ResimCacheBase.h"
 #include "Chaos/PBDJointConstraints.h"
@@ -12,68 +13,175 @@
 #define VALIDATE_REWIND_DATA 0
 #endif
 
+#ifndef DEBUG_NETWORK_PHYSICS
+#define DEBUG_NETWORK_PHYSICS 0
+#endif
+
+struct FBaseHistoryDatas;
+
 namespace Chaos
 {
 
-/** Base states history interface to store additional information required for rewind */
-struct FBaseStatesHistory
+/** Base rewind history used in the rewind datas */
+struct FBaseRewindHistory
 {
-	FORCEINLINE virtual ~FBaseStatesHistory() {}
+	FORCEINLINE virtual ~FBaseRewindHistory() {}
 
-	/** Rewind the states at some point in time */
-	FORCEINLINE virtual bool RewindStates(const int32 RewindFrame, const bool bResetSolver) { return true; }
-};
+	/** Set the package map for serialization */
+	FORCEINLINE virtual void SetPackageMap(class UPackageMap* InPackageMap) {}
 
-/** Base Inputs history interface to continuously apply Inputs inputs while resimulating */
-struct FBaseInputsHistory
-{
-	FORCEINLINE virtual ~FBaseInputsHistory() {}
+	/** Check if the history buffer contains an entry for the given frame*/
+	FORCEINLINE virtual bool HasValidDatas(const int32 ValidFrame) const { return false; }
+
+	/** Extract datas from the history buffer at a given time */
+	FORCEINLINE virtual bool ExtractDatas(const int32 ExtractFrame, const bool bResetSolver, void* HistoryDatas, const bool bExactFrame = false) { return true; }
+
+	/** Record datas into the history buffer at a given time */
+	FORCEINLINE virtual bool RecordDatas(const int32 RecordFrame, const void* HistoryDatas) { return true; }
+
+	/** Serialize the datas into an array of uint8 */
+	FORCEINLINE virtual void SerializeDatas(const uint32 StartFrame, const uint32 EndFrame, TArray<uint8>& ArchiveDatas, const int32 FrameOffset) const {}
 
-	/** Apply Inputs given a time */
-	FORCEINLINE virtual bool ApplyInputs(const int32 ApplyFrame, const bool bResetSolver) { return true; }
+	/** Deserialize the datas from the array of uint8 */
+	FORCEINLINE virtual void DeserializeDatas(const TArray<uint8>& ArchiveDatas, const int32 FrameOffset) {}
+
+	/** Debug the datas from the array of uint8 that will be transferred from client to server */
+	FORCEINLINE virtual void DebugDatas(const TArray<uint8>& ArchiveDatas, TArray<int32>& LocalFrames, TArray<int32>& ServerFrames, TArray<int32>& InputFrames) {}
+
+	/** Legacy interface to rewind states */
+	FORCEINLINE virtual bool RewindStates(const int32 RewindFrame, const bool bResetSolver) { return false; }
+
+	/** Legacy interface to apply inputs */
+	FORCEINLINE virtual bool ApplyInputs(const int32 ApplyFrame, const bool bResetSolver) { return false; }
 };
 
-/** Base history class that is holding a fixed size buffer of datas */
+/** Templated datas history holding a datas buffer */
 template<typename DatasType>
-struct FDatasRewindHistory 
+struct TDatasRewindHistory : public FBaseRewindHistory
 {
-	FORCEINLINE FDatasRewindHistory(const int32 FrameCount, const bool bIsHistoryLocal) : 
+	FORCEINLINE TDatasRewindHistory(const int32 FrameCount, const bool bIsHistoryLocal) :
 		bIsLocalHistory(bIsHistoryLocal), DatasArray(), CurrentFrame(0), CurrentIndex(0), NumFrames(FrameCount)
 	{
 		DatasArray.SetNum(NumFrames);
 	}
-	FORCEINLINE virtual ~FDatasRewindHistory() {}
+	FORCEINLINE virtual ~TDatasRewindHistory() {}
 
 protected:
 
-	/** Load the datas from the buffer at a specific frame */
-	FORCEINLINE bool LoadDatas(const int32 LoadFrame)
+	/** Get the closest (min/max) valid datas from the Datas frame */
+	FORCEINLINE int32 ClosestDatas(const int32 DatasFrame, const bool bMinDatas)
 	{
-		const int32 LocalFrame = LoadFrame % NumFrames;
-		if (LoadFrame == DatasArray[LocalFrame].DatasFrame)
+		int32 ClosestIndex = INDEX_NONE;
+		for (int32 FrameIndex = 0; FrameIndex < NumFrames; ++FrameIndex)
+		{
+			const int32 ValidFrame = bMinDatas ? FMath::Max(0, DatasFrame - FrameIndex) : DatasFrame + FrameIndex;
+			const int32 ValidIndex = ValidFrame % NumFrames;
+
+			if (DatasArray[ValidIndex].LocalFrame == ValidFrame)
+			{
+				ClosestIndex = ValidIndex;
+				break;
+			}
+		}
+		return ClosestIndex;
+	}
+
+public : 
+
+	/** Check if the history buffer contains an entry for the given frame*/
+	FORCEINLINE virtual bool HasValidDatas(const int32 ValidFrame) const override
+	{
+		const int32 LocalFrame = ValidFrame % NumFrames;
+		return ValidFrame == DatasArray[LocalFrame].LocalFrame;
+	}
+
+	/** Extract states at a given time */
+	FORCEINLINE virtual bool ExtractDatas(const int32 ExtractFrame, const bool bResetSolver, void* HistoryDatas, const bool bExactFrame = false) override
+	{
+		const int32 LocalFrame = ExtractFrame % NumFrames;
+		if (ExtractFrame == DatasArray[LocalFrame].LocalFrame)
 		{
-			CurrentFrame = LoadFrame;
+			CurrentFrame = ExtractFrame;
 			CurrentIndex = LocalFrame;
+
+#if DEBUG_NETWORK_PHYSICS
+			UE_LOG(LogTemp, Log, TEXT("		Found matching datas into history at frame %d"), ExtractFrame);
+#endif
+			*static_cast<DatasType*>(HistoryDatas) = DatasArray[CurrentIndex];
 			return true;
 		}
+		else if(!bExactFrame)
+		{
+			if (bResetSolver)
+			{
+				UE_LOG(LogChaos, Warning, TEXT("		Unable to extract datas at frame %d while rewinding the simulation"), ExtractFrame);
+			}
+			const int32 MinFrame = ClosestDatas(ExtractFrame, true);
+			const int32 MaxFrame = ClosestDatas(ExtractFrame, false);
+
+			if (MinFrame != INDEX_NONE && MaxFrame != INDEX_NONE)
+			{
+				static_cast<DatasType*>(HistoryDatas)->InterpolateDatas(
+					DatasArray[MinFrame], DatasArray[MaxFrame]);
+
+				const int32 DeltaFrame = FMath::Abs(ExtractFrame - DatasArray[MinFrame].LocalFrame);
+
+				static_cast<DatasType*>(HistoryDatas)->LocalFrame = ExtractFrame;
+				static_cast<DatasType*>(HistoryDatas)->ServerFrame = DatasArray[MinFrame].ServerFrame + DeltaFrame;
+				static_cast<DatasType*>(HistoryDatas)->InputFrame = DatasArray[MinFrame].InputFrame + DeltaFrame;
+				return true;
+#if DEBUG_NETWORK_PHYSICS
+				UE_LOG(LogTemp, Log, TEXT("		Smoothing datas between frame %d and %d - > [%d %d]"), DatasArray[MinFrame].LocalFrame, DatasArray[MaxFrame].LocalFrame, static_cast<DatasType*>(HistoryDatas)->InputFrame, static_cast<DatasType*>(HistoryDatas)->ServerFrame);
+#endif
+			}
+			else if (MinFrame != INDEX_NONE)
+			{
+				*static_cast<DatasType*>(HistoryDatas) = DatasArray[MinFrame];
+				return true;
+#if DEBUG_NETWORK_PHYSICS
+				UE_LOG(LogTemp, Log, TEXT("		Setting datas to frame %d"), DatasArray[MinFrame].LocalFrame);
+#endif
+			}
+			else
+			{
+#if DEBUG_NETWORK_PHYSICS
+				UE_LOG(LogTemp, Log, TEXT("		Failed to find datas bounds : Min = %d | Max = %d"), MinFrame, MaxFrame);
+#endif
+				return false;
+			}
+		}
 		return false;
 	}
 
-	/** Read directly the datas from the buffer at a specific frame without chacking the datasframe */
-	FORCEINLINE void ReadDatas(const int32 LoadFrame)
+	/** Load the datas from the buffer at a specific frame */
+	FORCEINLINE bool LoadDatas(const int32 LoadFrame)
 	{
 		const int32 LocalFrame = LoadFrame % NumFrames;
+		DatasArray[LocalFrame].LocalFrame = LoadFrame;
 		CurrentFrame = LoadFrame;
 		CurrentIndex = LocalFrame;
+		return true;
+	}
+
+
+	/** Eval the datas from the buffer at a specific frame */
+	FORCEINLINE bool EvalDatas(const int32 EvalFrame)
+	{
+		const int32 LocalFrame = EvalFrame % NumFrames;
+		if (EvalFrame == DatasArray[LocalFrame].LocalFrame)
+		{
+			CurrentFrame = EvalFrame;
+			CurrentIndex = LocalFrame;
+			return true;
+		}
+		return false;
 	}
 
 	/** Record the datas from the buffer at a specific frame */
-	FORCEINLINE virtual bool RecordDatas(const int32 RecordFrame)
+	FORCEINLINE virtual bool RecordDatas(const int32 RecordFrame, const void* HistoryDatas) override
 	{
-		const int32 LocalFrame = RecordFrame % NumFrames;
-		DatasArray[LocalFrame].DatasFrame = RecordFrame;
-		CurrentFrame = RecordFrame;
-		CurrentIndex = LocalFrame;
+		LoadDatas(RecordFrame);
+		DatasArray[CurrentIndex] = *static_cast<const DatasType*>(HistoryDatas);
 		return true;
 	}
 
@@ -81,6 +189,23 @@ protected:
 	DatasType& GetCurrentDatas() { return DatasArray[CurrentIndex]; }
 	const DatasType& GetCurrentDatas() const { return DatasArray[CurrentIndex]; }
 
+	/** Get the number of valid datas in the buffer index range */
+	FORCEINLINE uint32 NumValidDatas(const uint32 StartFrame, const uint32 EndFrame) const
+	{
+		uint32 NumDatas = 0;
+		for (uint32 FrameIndex = StartFrame; FrameIndex < EndFrame; ++FrameIndex)
+		{
+			const int32 LocalFrame = FrameIndex % NumFrames;
+			if (FrameIndex == DatasArray[LocalFrame].LocalFrame)
+			{
+				++NumDatas;
+			}
+		}
+		return NumDatas;
+	}
+
+protected : 
+
 	/** Check if the history is on the local/remote client*/
 	bool bIsLocalHistory;
 
@@ -97,60 +222,6 @@ protected:
 	int32 NumFrames = 0;
 };
 
-/** Base history class for the states rewind datas */
-template<typename DatasType>
-struct  FStatesRewindHistory : public FDatasRewindHistory<DatasType>, public FBaseStatesHistory
-{
-	using DatasSuper = FDatasRewindHistory<DatasType>;
-
-	FORCEINLINE FStatesRewindHistory(const int32 FrameCount, const bool bIsHistoryLocal) :
-		FDatasRewindHistory<DatasType>(FrameCount, bIsHistoryLocal), FBaseStatesHistory()
-	{}
-
-	FORCEINLINE virtual ~FStatesRewindHistory() override {}
-
-	/** Rewind the states */
-	FORCEINLINE virtual bool RewindStates(const int32 RewindFrame, const bool bResetSolver) override
-	{
-		if (!DatasSuper::LoadDatas(RewindFrame))
-		{
-			if(bResetSolver)
-			{
-				UE_LOG(LogChaos, Warning, TEXT("Unable to read target state at frame %d while rewinding the simulation"), RewindFrame);
-			}
-			return false;
-		}
-		return true;
-	}
-};
-
-/** Base history class for the Inputs rewind datas */
-template<typename DatasType>
-struct  FInputsRewindHistory : public FDatasRewindHistory<DatasType>, public FBaseInputsHistory
-{
-	using DatasSuper = FDatasRewindHistory<DatasType>;
-
-	FORCEINLINE FInputsRewindHistory(const int32 FrameCount, const bool bIsHistoryLocal) :
-		FDatasRewindHistory<DatasType>(FrameCount, bIsHistoryLocal), FBaseInputsHistory()
-	{}
-
-	FORCEINLINE virtual ~FInputsRewindHistory() override {}
-
-	/** Apply the Inputs */
-	FORCEINLINE virtual bool ApplyInputs(const int32 ApplyFrame, const bool bResetSolver) override
-	{
-		if (!DatasSuper::LoadDatas(ApplyFrame))
-		{
-			if (DatasSuper::bIsLocalHistory || (!DatasSuper::bIsLocalHistory && bResetSolver))
-			{
-				UE_LOG(LogChaos, Warning, TEXT("Unable to apply Inputs on the local client at frame %d"), ApplyFrame);
-			}
-			return false;
-		}
-		return true;
-	}
-};
-
 struct FFrameAndPhase
 {
 	enum EParticleHistoryPhase : uint8
@@ -348,6 +419,59 @@ public:
 		return NoEntryInSync<THandle, T, bNoEntryIsHead>::Helper(Handle);
 	}
 
+	T& Insert(const FFrameAndPhase FrameAndPhase, FDirtyPropertiesPool& Manager)
+	{
+		T* Result = nullptr;
+
+		int32 FrameIndex = FindIdx(FrameAndPhase);
+		if(FrameIndex != INDEX_NONE)
+		{
+			Result = &GetPool(Manager).GetElement(Buffer[FrameIndex].Ref);
+		}
+		else
+		{
+			FPropertyIdx ElementRef;
+			if(Next >= Buffer.Num())
+			{
+				GetPool(Manager).AddElement(ElementRef);
+				Buffer.Add({ ElementRef, FrameAndPhase });
+			}
+			else
+			{
+				ElementRef = Buffer[Next].Ref;
+			}
+			Result = &GetPool(Manager).GetElement(ElementRef);
+
+			int32 PrevFrame = Next;
+			int32 NextFrame = PrevFrame;
+			for (int32 Count = 0; Count < NumValid; ++Count)
+			{
+				NextFrame = PrevFrame;
+
+				--PrevFrame;
+				if (PrevFrame < 0) { PrevFrame = Buffer.Num() - 1; }
+
+				const FPropertyInterval& PrevInterval = Buffer[PrevFrame];
+				if (PrevInterval.FrameAndPhase < FrameAndPhase)
+				{
+					Buffer[NextFrame].FrameAndPhase = FrameAndPhase;
+					Buffer[NextFrame].Ref = ElementRef;
+					break;
+				}
+				else
+				{
+					Buffer[NextFrame] = Buffer[PrevFrame];
+				}
+			}
+
+			++Next;
+			if (Next == Capacity) { Next = 0; }
+
+			NumValid = FMath::Min(++NumValid, Capacity);
+		}
+		return *Result;
+	}
+
 private:
 
 	const int32 FindIdx(const FFrameAndPhase FrameAndPhase) const
@@ -693,6 +817,7 @@ struct FGeometryParticleStateBase
 	, KinematicTarget(ComputeCircularSize(NumFrames))
 	, TargetPositions(ComputeCircularSize(NumFrames))
 	, TargetVelocities(ComputeCircularSize(NumFrames))
+	, TargetStates(ComputeCircularSize(NumFrames))
 	{
 
 	}
@@ -712,6 +837,7 @@ struct FGeometryParticleStateBase
 		KinematicTarget.Release(Manager);
 		TargetPositions.Release(Manager);
 		TargetVelocities.Release(Manager);
+		TargetStates.Release(Manager);
 	}
 
 	void Reset()
@@ -725,6 +851,7 @@ struct FGeometryParticleStateBase
 		KinematicTarget.Reset();
 		TargetVelocities.Reset();
 		TargetPositions.Reset();
+		TargetStates.Reset();
 	}
 
 	void ClearEntryAndFuture(const FFrameAndPhase FrameAndPhase)
@@ -738,6 +865,7 @@ struct FGeometryParticleStateBase
 		KinematicTarget.ClearEntryAndFuture(FrameAndPhase);
 		TargetPositions.ClearEntryAndFuture(FrameAndPhase);
 		TargetVelocities.ClearEntryAndFuture(FrameAndPhase);
+		TargetStates.ClearEntryAndFuture(FrameAndPhase);
 	}
 
 	void ExtractHistoryState(int32& PositionValidCount, int32& VelocityValidCount, int32& PositionNextIterator, int32& VelocityNextIterator) const
@@ -799,6 +927,7 @@ struct FGeometryParticleStateBase
 
 	TParticlePropertyBuffer<FParticlePositionRotation, EChaosProperty::XR> TargetPositions;
 	TParticlePropertyBuffer<FParticleVelocities, EChaosProperty::Velocities> TargetVelocities;
+	TParticlePropertyBuffer<FParticleDynamicMisc, EChaosProperty::DynamicMisc> TargetStates;
 
 	FShapesArrayStateBase ShapesArrayState;
 
@@ -1263,25 +1392,25 @@ public:
 	void CHAOS_API SpawnProxyIfNeeded(FSingleParticlePhysicsProxy& Proxy);
 
 	/** Add Inputs history to the rewind datas for future use while resimulating */
-	void CHAOS_API AddInputsHistory(const TSharedPtr<FBaseInputsHistory>& InputsHistory)
+	void CHAOS_API AddInputsHistory(const TSharedPtr<FBaseRewindHistory>& InputsHistory)
 	{
 		InputsHistories.Add(InputsHistory.ToWeakPtr());
 	}
 
-	/** Remove Inputs history from the rewind datasg */
-	void CHAOS_API RemoveInputsHistory(const TSharedPtr<FBaseInputsHistory>& InputsHistory)
+	/** Remove Inputs history from the rewind datas */
+	void CHAOS_API RemoveInputsHistory(const TSharedPtr<FBaseRewindHistory>& InputsHistory)
 	{
 		InputsHistories.Remove(InputsHistory.ToWeakPtr());
 	}
 
 	/** Add states history to the rewind datas for future use while rewinding */
-	void CHAOS_API AddStatesHistory(const TSharedPtr<FBaseStatesHistory>& StatesHistory)
+	void CHAOS_API AddStatesHistory(const TSharedPtr<FBaseRewindHistory>& StatesHistory)
 	{
 		StatesHistories.Add(StatesHistory.ToWeakPtr());
 	}
 
 	/** Remove states history from the rewind datas */
-	void CHAOS_API RemoveStatesHistory(const TSharedPtr<FBaseStatesHistory>& StatesHistory)
+	void CHAOS_API RemoveStatesHistory(const TSharedPtr<FBaseRewindHistory>& StatesHistory)
 	{
 		StatesHistories.Remove(StatesHistory.ToWeakPtr());
 	}
@@ -1297,6 +1426,9 @@ public:
 	/** Return the rewind data solver */
 	const FPBDRigidsSolver* GetSolver() const { return Solver; }
 
+	/** Find the first previous valid frame having received physics target from the server */
+	int32 CHAOS_API FindValidResimFrame(const int32 ResimFrame);
+
 private:
 
 	friend class FPBDRigidsSolver;
@@ -1443,6 +1575,8 @@ private:
 	}
 
 	bool RewindToFrame(int32 RewindFrame);
+	
+
 
 	/** Apply targets positions and velocities while resimulating */
 	void ApplyTargets(const int32 Frame, const bool bResetSimulation);
@@ -1461,8 +1595,8 @@ private:
 	TDirtyObjects<FDirtyJointInfo> DirtyJoints;
 	TDirtyObjects<FDirtyParticleErrorInfo> DirtyParticleErrors;
 
-	TArray<TWeakPtr<FBaseInputsHistory>> InputsHistories;
-	TArray<TWeakPtr<FBaseStatesHistory>> StatesHistories;
+	TArray<TWeakPtr<FBaseRewindHistory>> InputsHistories;
+	TArray<TWeakPtr<FBaseRewindHistory>> StatesHistories;
 
 	FPBDRigidsSolver* Solver;
 	int32 CurFrame;
@@ -1497,6 +1631,9 @@ public:
 	*   See ISimCallbackObject for recording inputs to callbacks associated with this PhysicsStep */
 	virtual void ProcessInputs_Internal(int32 PhysicsStep, const TArray<FSimCallbackInputAndObject>& SimCallbackInputs){}
 
+	/** Called after any presim callbacks are triggered and after physics data has marshalled over in order to modify the sim callback outputs */
+	virtual void ApplyCallbacks_Internal(int32 PhysicsStep, const TArray<ISimCallbackObject*>& SimCallbackObjects) {}
+
 	/** Called before any inputs are marshalled over to the physics thread.
 	*	The physics state has not been applied yet, and cannot be inspected anyway because this is triggered from the external thread (game thread)
 	*	Gives user the ability to modify inputs or record them - this can help with reducing latency if you want to act on inputs immediately
@@ -1524,11 +1661,15 @@ public:
 	/** Called after each rewind step. This is to give user code the opportunity to trigger other code after each rewind step
 	*   Usually to simulate external systems that ran in lock step with the physics sim
 	*/
-	virtual void PostResimStep_Internal(int32 PhysicsStep){}
+	virtual void PostResimStep_Internal(int32 PhysicsStep) {}
 
-	virtual void RegisterRewindableSimCallback_Internal(ISimCallbackObject* Callback) { ensure(false); }
+	/** Register a sim callback onto the rewind callback */
+	virtual void RegisterRewindableSimCallback_Internal(ISimCallbackObject* Callback) {}
 
 	/** Called When resim is finished with debug information about the resim */
 	virtual void SetResimDebugInfo_Internal(const FResimDebugInfo& ResimDebugInfo){}
+
+	/** Rewind Data holding the callback */
+	Chaos::FRewindData* RewindData = nullptr;
 };
-}
+}
\ No newline at end of file
