Update Disable Optimization macro

diff --git a/Engine/Plugins/Experimental/ChaosVehiclesPlugin/Source/ChaosVehicles/Private/ChaosVehicleMovementComponent.cpp b/Engine/Plugins/Experimental/ChaosVehiclesPlugin/Source/ChaosVehicles/Private/ChaosVehicleMovementComponent.cpp
--- a/Engine/Plugins/Experimental/ChaosVehiclesPlugin/Source/ChaosVehicles/Private/ChaosVehicleMovementComponent.cpp
+++ b/Engine/Plugins/Experimental/ChaosVehiclesPlugin/Source/ChaosVehicles/Private/ChaosVehicleMovementComponent.cpp
@@ -44,7 +44,7 @@
 #define LOCTEXT_NAMESPACE "UVehicleMovementComponent"
 
 #if VEHICLE_DEBUGGING_ENABLED
-PRAGMA_DISABLE_OPTIMIZATION
+UE_DISABLE_OPTIMIZATION
 #endif
 
 DEFINE_LOG_CATEGORY(LogVehicle);
@@ -1959,7 +1959,7 @@ void UChaosVehicleMovementComponent::PutAllEnabledRigidBodiesToSleep()
 
 
 #if VEHICLE_DEBUGGING_ENABLED
-PRAGMA_ENABLE_OPTIMIZATION
+UE_ENABLE_OPTIMIZATION
 #endif
 
 
diff --git a/Engine/Plugins/Experimental/ChaosVehiclesPlugin/Source/ChaosVehicles/Private/ChaosVehicleWheel.cpp b/Engine/Plugins/Experimental/ChaosVehiclesPlugin/Source/ChaosVehicles/Private/ChaosVehicleWheel.cpp
--- a/Engine/Plugins/Experimental/ChaosVehiclesPlugin/Source/ChaosVehicles/Private/ChaosVehicleWheel.cpp
+++ b/Engine/Plugins/Experimental/ChaosVehiclesPlugin/Source/ChaosVehicles/Private/ChaosVehicleWheel.cpp
@@ -11,7 +11,7 @@
 #include "ChaosWheeledVehicleMovementComponent.h"
 
 #if VEHICLE_DEBUGGING_ENABLED
-PRAGMA_DISABLE_OPTIMIZATION
+UE_DISABLE_OPTIMIZATION
 #endif
 
 
@@ -169,7 +169,7 @@ UPhysicalMaterial* UChaosVehicleWheel::GetContactSurfaceMaterial()
 
 
 #if VEHICLE_DEBUGGING_ENABLED
-PRAGMA_ENABLE_OPTIMIZATION
+UE_ENABLE_OPTIMIZATION
 #endif
 
 
diff --git a/Engine/Plugins/Experimental/ChaosVehiclesPlugin/Source/ChaosVehicles/Private/ChaosWheeledVehicleMovementComponent.cpp b/Engine/Plugins/Experimental/ChaosVehiclesPlugin/Source/ChaosVehicles/Private/ChaosWheeledVehicleMovementComponent.cpp
--- a/Engine/Plugins/Experimental/ChaosVehiclesPlugin/Source/ChaosVehicles/Private/ChaosWheeledVehicleMovementComponent.cpp
+++ b/Engine/Plugins/Experimental/ChaosVehiclesPlugin/Source/ChaosVehicles/Private/ChaosWheeledVehicleMovementComponent.cpp
@@ -33,7 +33,7 @@
 #include "PhysicsProxy/SingleParticlePhysicsProxy.h"
 
 #if VEHICLE_DEBUGGING_ENABLED
-PRAGMA_DISABLE_OPTIMIZATION
+UE_DISABLE_OPTIMIZATION
 #endif
 
 DECLARE_STATS_GROUP(TEXT("ChaosVehicle"), STATGROUP_ChaosVehicle, STATGROUP_Advanced);
@@ -2916,7 +2916,7 @@ FChaosWheelSetup::FChaosWheelSetup()
 }
 
 #if VEHICLE_DEBUGGING_ENABLED
-PRAGMA_ENABLE_OPTIMIZATION
+UE_ENABLE_OPTIMIZATION
 #endif
 
 
diff --git a/Engine/Plugins/Experimental/ChaosVehiclesPlugin/Source/ChaosVehicles/Public/ChaosWheeledVehicleMovementComponent.h b/Engine/Plugins/Experimental/ChaosVehiclesPlugin/Source/ChaosVehicles/Public/ChaosWheeledVehicleMovementComponent.h
--- a/Engine/Plugins/Experimental/ChaosVehiclesPlugin/Source/ChaosVehicles/Public/ChaosWheeledVehicleMovementComponent.h
+++ b/Engine/Plugins/Experimental/ChaosVehiclesPlugin/Source/ChaosVehicles/Public/ChaosWheeledVehicleMovementComponent.h
@@ -13,7 +13,7 @@
 #include "ChaosWheeledVehicleMovementComponent.generated.h"
 
 #if VEHICLE_DEBUGGING_ENABLED
-PRAGMA_DISABLE_OPTIMIZATION
+UE_DISABLE_OPTIMIZATION
 #endif
 
 struct CHAOSVEHICLES_API FWheeledVehicleDebugParams
@@ -961,5 +961,5 @@ protected:
 };
 
 #if VEHICLE_DEBUGGING_ENABLED
-PRAGMA_ENABLE_OPTIMIZATION
+UE_ENABLE_OPTIMIZATION
 #endif
diff --git a/Engine/Source/Runtime/Experimental/ChaosVehicles/ChaosVehiclesCore/Private/AerodynamicsSystem.cpp b/Engine/Source/Runtime/Experimental/ChaosVehicles/ChaosVehiclesCore/Private/AerodynamicsSystem.cpp
--- a/Engine/Source/Runtime/Experimental/ChaosVehicles/ChaosVehiclesCore/Private/AerodynamicsSystem.cpp
+++ b/Engine/Source/Runtime/Experimental/ChaosVehicles/ChaosVehiclesCore/Private/AerodynamicsSystem.cpp
@@ -3,7 +3,7 @@
 #include "AerodynamicsSystem.h"
 
 #if VEHICLE_DEBUGGING_ENABLED
-PRAGMA_DISABLE_OPTIMIZATION
+UE_DISABLE_OPTIMIZATION
 #endif
 
 namespace Chaos
@@ -31,5 +31,5 @@ namespace Chaos
 } // namespace Chaos
 
 #if VEHICLE_DEBUGGING_ENABLED
-PRAGMA_ENABLE_OPTIMIZATION
+UE_ENABLE_OPTIMIZATION
 #endif
diff --git a/Engine/Source/Runtime/Experimental/ChaosVehicles/ChaosVehiclesCore/Private/AerofoilSystem.cpp b/Engine/Source/Runtime/Experimental/ChaosVehicles/ChaosVehiclesCore/Private/AerofoilSystem.cpp
--- a/Engine/Source/Runtime/Experimental/ChaosVehicles/ChaosVehiclesCore/Private/AerofoilSystem.cpp
+++ b/Engine/Source/Runtime/Experimental/ChaosVehicles/ChaosVehiclesCore/Private/AerofoilSystem.cpp
@@ -3,7 +3,7 @@
 #include "AerofoilSystem.h"
 
 #if VEHICLE_DEBUGGING_ENABLED
-PRAGMA_DISABLE_OPTIMIZATION
+UE_DISABLE_OPTIMIZATION
 #endif
 
 namespace Chaos
@@ -134,5 +134,5 @@ namespace Chaos
 } // namespace Chaos
 
 #if VEHICLE_DEBUGGING_ENABLED
-PRAGMA_ENABLE_OPTIMIZATION
+UE_ENABLE_OPTIMIZATION
 #endif
diff --git a/Engine/Source/Runtime/Experimental/ChaosVehicles/ChaosVehiclesCore/Private/ArcadeSystem.cpp b/Engine/Source/Runtime/Experimental/ChaosVehicles/ChaosVehiclesCore/Private/ArcadeSystem.cpp
--- a/Engine/Source/Runtime/Experimental/ChaosVehicles/ChaosVehiclesCore/Private/ArcadeSystem.cpp
+++ b/Engine/Source/Runtime/Experimental/ChaosVehicles/ChaosVehiclesCore/Private/ArcadeSystem.cpp
@@ -3,7 +3,7 @@
 #include "ArcadeSystem.h"
 
 #if VEHICLE_DEBUGGING_ENABLED
-PRAGMA_DISABLE_OPTIMIZATION
+UE_DISABLE_OPTIMIZATION
 #endif
 
 namespace Chaos
@@ -30,5 +30,5 @@ namespace Chaos
 } // namespace Chaos
 
 #if VEHICLE_DEBUGGING_ENABLED
-PRAGMA_ENABLE_OPTIMIZATION
+UE_ENABLE_OPTIMIZATION
 #endif
diff --git a/Engine/Source/Runtime/Experimental/ChaosVehicles/ChaosVehiclesCore/Private/EngineSystem.cpp b/Engine/Source/Runtime/Experimental/ChaosVehicles/ChaosVehiclesCore/Private/EngineSystem.cpp
--- a/Engine/Source/Runtime/Experimental/ChaosVehicles/ChaosVehiclesCore/Private/EngineSystem.cpp
+++ b/Engine/Source/Runtime/Experimental/ChaosVehicles/ChaosVehiclesCore/Private/EngineSystem.cpp
@@ -3,7 +3,7 @@
 #include "EngineSystem.h"
 
 #if VEHICLE_DEBUGGING_ENABLED
-PRAGMA_DISABLE_OPTIMIZATION
+UE_DISABLE_OPTIMIZATION
 #endif
 
 namespace Chaos
@@ -80,5 +80,5 @@ namespace Chaos
 
 
 #if VEHICLE_DEBUGGING_ENABLED
-PRAGMA_ENABLE_OPTIMIZATION
+UE_ENABLE_OPTIMIZATION
 #endif
diff --git a/Engine/Source/Runtime/Experimental/ChaosVehicles/ChaosVehiclesCore/Private/ThrustSystem.cpp b/Engine/Source/Runtime/Experimental/ChaosVehicles/ChaosVehiclesCore/Private/ThrustSystem.cpp
--- a/Engine/Source/Runtime/Experimental/ChaosVehicles/ChaosVehiclesCore/Private/ThrustSystem.cpp
+++ b/Engine/Source/Runtime/Experimental/ChaosVehicles/ChaosVehiclesCore/Private/ThrustSystem.cpp
@@ -3,7 +3,7 @@
 #include "ThrustSystem.h"
 
 #if VEHICLE_DEBUGGING_ENABLED
-PRAGMA_DISABLE_OPTIMIZATION
+UE_DISABLE_OPTIMIZATION
 #endif
 
 namespace Chaos
@@ -49,5 +49,5 @@ namespace Chaos
 } // namespace Chaos
 
 #if VEHICLE_DEBUGGING_ENABLED
-PRAGMA_ENABLE_OPTIMIZATION
+UE_ENABLE_OPTIMIZATION
 #endif
diff --git a/Engine/Source/Runtime/Experimental/ChaosVehicles/ChaosVehiclesCore/Private/TransmissionSystem.cpp b/Engine/Source/Runtime/Experimental/ChaosVehicles/ChaosVehiclesCore/Private/TransmissionSystem.cpp
--- a/Engine/Source/Runtime/Experimental/ChaosVehicles/ChaosVehiclesCore/Private/TransmissionSystem.cpp
+++ b/Engine/Source/Runtime/Experimental/ChaosVehicles/ChaosVehiclesCore/Private/TransmissionSystem.cpp
@@ -3,7 +3,7 @@
 #include "TransmissionSystem.h"
 
 #if VEHICLE_DEBUGGING_ENABLED
-PRAGMA_DISABLE_OPTIMIZATION
+UE_DISABLE_OPTIMIZATION
 #endif
 
 namespace Chaos
@@ -102,5 +102,5 @@ namespace Chaos
 } // namespace Chaos
 
 #if VEHICLE_DEBUGGING_ENABLED
-PRAGMA_ENABLE_OPTIMIZATION
+UE_ENABLE_OPTIMIZATION
 #endif
diff --git a/Engine/Source/Runtime/Experimental/ChaosVehicles/ChaosVehiclesCore/Private/VehicleUtility.cpp b/Engine/Source/Runtime/Experimental/ChaosVehicles/ChaosVehiclesCore/Private/VehicleUtility.cpp
--- a/Engine/Source/Runtime/Experimental/ChaosVehicles/ChaosVehiclesCore/Private/VehicleUtility.cpp
+++ b/Engine/Source/Runtime/Experimental/ChaosVehicles/ChaosVehiclesCore/Private/VehicleUtility.cpp
@@ -4,7 +4,7 @@
 #include "HAL/PlatformTime.h"
 
 #if VEHICLE_DEBUGGING_ENABLED
-PRAGMA_DISABLE_OPTIMIZATION
+UE_DISABLE_OPTIMIZATION
 #endif
 
 namespace Chaos
@@ -225,5 +225,5 @@ namespace Chaos
 } // namespace Chaos
 
 #if VEHICLE_DEBUGGING_ENABLED
-PRAGMA_ENABLE_OPTIMIZATION
+UE_ENABLE_OPTIMIZATION
 #endif
diff --git a/Engine/Source/Runtime/Experimental/ChaosVehicles/ChaosVehiclesCore/Private/WheelSystem.cpp b/Engine/Source/Runtime/Experimental/ChaosVehicles/ChaosVehiclesCore/Private/WheelSystem.cpp
--- a/Engine/Source/Runtime/Experimental/ChaosVehicles/ChaosVehiclesCore/Private/WheelSystem.cpp
+++ b/Engine/Source/Runtime/Experimental/ChaosVehicles/ChaosVehiclesCore/Private/WheelSystem.cpp
@@ -3,7 +3,7 @@
 #include "WheelSystem.h"
 
 #if VEHICLE_DEBUGGING_ENABLED
-PRAGMA_DISABLE_OPTIMIZATION
+UE_DISABLE_OPTIMIZATION
 #endif
 
 namespace Chaos
@@ -257,5 +257,5 @@ namespace Chaos
 } // namespace Chaos
 
 #if VEHICLE_DEBUGGING_ENABLED
-PRAGMA_ENABLE_OPTIMIZATION
+UE_ENABLE_OPTIMIZATION
 #endif
diff --git a/Engine/Source/Runtime/Experimental/ChaosVehicles/ChaosVehiclesCore/Public/AerodynamicsSystem.h b/Engine/Source/Runtime/Experimental/ChaosVehicles/ChaosVehiclesCore/Public/AerodynamicsSystem.h
--- a/Engine/Source/Runtime/Experimental/ChaosVehicles/ChaosVehiclesCore/Public/AerodynamicsSystem.h
+++ b/Engine/Source/Runtime/Experimental/ChaosVehicles/ChaosVehiclesCore/Public/AerodynamicsSystem.h
@@ -7,7 +7,7 @@
 #include "VehicleUtility.h"
 
 #if VEHICLE_DEBUGGING_ENABLED
-PRAGMA_DISABLE_OPTIMIZATION
+UE_DISABLE_OPTIMIZATION
 #endif
 
 namespace Chaos
@@ -87,5 +87,5 @@ namespace Chaos
 } // namespace Chaos
 
 #if VEHICLE_DEBUGGING_ENABLED
-PRAGMA_ENABLE_OPTIMIZATION
+UE_ENABLE_OPTIMIZATION
 #endif
diff --git a/Engine/Source/Runtime/Experimental/ChaosVehicles/ChaosVehiclesCore/Public/AerofoilSystem.h b/Engine/Source/Runtime/Experimental/ChaosVehicles/ChaosVehiclesCore/Public/AerofoilSystem.h
--- a/Engine/Source/Runtime/Experimental/ChaosVehicles/ChaosVehiclesCore/Public/AerofoilSystem.h
+++ b/Engine/Source/Runtime/Experimental/ChaosVehicles/ChaosVehiclesCore/Public/AerofoilSystem.h
@@ -9,7 +9,7 @@
 #include "VehicleUtility.h"
 
 #if VEHICLE_DEBUGGING_ENABLED
-PRAGMA_DISABLE_OPTIMIZATION
+UE_DISABLE_OPTIMIZATION
 #endif
 
 namespace Chaos
@@ -129,5 +129,5 @@ namespace Chaos
 } // namespace Chaos
 
 #if VEHICLE_DEBUGGING_ENABLED
-PRAGMA_ENABLE_OPTIMIZATION
+UE_ENABLE_OPTIMIZATION
 #endif
\ No newline at end of file
diff --git a/Engine/Source/Runtime/Experimental/ChaosVehicles/ChaosVehiclesCore/Public/ArcadeSystem.h b/Engine/Source/Runtime/Experimental/ChaosVehicles/ChaosVehiclesCore/Public/ArcadeSystem.h
--- a/Engine/Source/Runtime/Experimental/ChaosVehicles/ChaosVehiclesCore/Public/ArcadeSystem.h
+++ b/Engine/Source/Runtime/Experimental/ChaosVehicles/ChaosVehiclesCore/Public/ArcadeSystem.h
@@ -6,7 +6,7 @@
 #include "VehicleUtility.h"
 
 #if VEHICLE_DEBUGGING_ENABLED
-PRAGMA_DISABLE_OPTIMIZATION
+UE_DISABLE_OPTIMIZATION
 #endif
 
 namespace Chaos
@@ -166,5 +166,5 @@ namespace Chaos
 } // namespace Chaos
 
 #if VEHICLE_DEBUGGING_ENABLED
-PRAGMA_ENABLE_OPTIMIZATION
+UE_ENABLE_OPTIMIZATION
 #endif
\ No newline at end of file
diff --git a/Engine/Source/Runtime/Experimental/ChaosVehicles/ChaosVehiclesCore/Public/EngineSystem.h b/Engine/Source/Runtime/Experimental/ChaosVehicles/ChaosVehiclesCore/Public/EngineSystem.h
--- a/Engine/Source/Runtime/Experimental/ChaosVehicles/ChaosVehiclesCore/Public/EngineSystem.h
+++ b/Engine/Source/Runtime/Experimental/ChaosVehicles/ChaosVehiclesCore/Public/EngineSystem.h
@@ -7,7 +7,7 @@
 #include "VehicleUtility.h"
 
 #if VEHICLE_DEBUGGING_ENABLED
-PRAGMA_DISABLE_OPTIMIZATION
+UE_DISABLE_OPTIMIZATION
 #endif
 
 
@@ -168,5 +168,5 @@ namespace Chaos
 } // namespace Chaos
 
 #if VEHICLE_DEBUGGING_ENABLED
-PRAGMA_ENABLE_OPTIMIZATION
+UE_ENABLE_OPTIMIZATION
 #endif
diff --git a/Engine/Source/Runtime/Experimental/ChaosVehicles/ChaosVehiclesCore/Public/SimpleVehicle.h b/Engine/Source/Runtime/Experimental/ChaosVehicles/ChaosVehiclesCore/Public/SimpleVehicle.h
--- a/Engine/Source/Runtime/Experimental/ChaosVehicles/ChaosVehiclesCore/Public/SimpleVehicle.h
+++ b/Engine/Source/Runtime/Experimental/ChaosVehicles/ChaosVehiclesCore/Public/SimpleVehicle.h
@@ -16,7 +16,7 @@
 //#include "BallastSystem.h"
 
 #if VEHICLE_DEBUGGING_ENABLED
-PRAGMA_DISABLE_OPTIMIZATION
+UE_DISABLE_OPTIMIZATION
 #endif
 
 namespace Chaos
@@ -193,5 +193,5 @@ public:
 } // namespace Chaos
 
 #if VEHICLE_DEBUGGING_ENABLED
-PRAGMA_ENABLE_OPTIMIZATION
+UE_ENABLE_OPTIMIZATION
 #endif
diff --git a/Engine/Source/Runtime/Experimental/ChaosVehicles/ChaosVehiclesCore/Public/SteeringSystem.h b/Engine/Source/Runtime/Experimental/ChaosVehicles/ChaosVehiclesCore/Public/SteeringSystem.h
--- a/Engine/Source/Runtime/Experimental/ChaosVehicles/ChaosVehiclesCore/Public/SteeringSystem.h
+++ b/Engine/Source/Runtime/Experimental/ChaosVehicles/ChaosVehiclesCore/Public/SteeringSystem.h
@@ -9,7 +9,7 @@
 #include "VehicleUtility.h"
 
 #if VEHICLE_DEBUGGING_ENABLED
-PRAGMA_DISABLE_OPTIMIZATION
+UE_DISABLE_OPTIMIZATION
 #endif
 
 namespace Chaos
@@ -131,5 +131,5 @@ namespace Chaos
 } // namespace Chaos
 
 #if VEHICLE_DEBUGGING_ENABLED
-PRAGMA_ENABLE_OPTIMIZATION
+UE_ENABLE_OPTIMIZATION
 #endif
\ No newline at end of file
diff --git a/Engine/Source/Runtime/Experimental/ChaosVehicles/ChaosVehiclesCore/Public/SteeringUtility.h b/Engine/Source/Runtime/Experimental/ChaosVehicles/ChaosVehiclesCore/Public/SteeringUtility.h
--- a/Engine/Source/Runtime/Experimental/ChaosVehicles/ChaosVehiclesCore/Public/SteeringUtility.h
+++ b/Engine/Source/Runtime/Experimental/ChaosVehicles/ChaosVehiclesCore/Public/SteeringUtility.h
@@ -7,7 +7,7 @@
 #include "VehicleUtility.h"
 
 #if VEHICLE_DEBUGGING_ENABLED
-PRAGMA_DISABLE_OPTIMIZATION
+UE_DISABLE_OPTIMIZATION
 #endif
 
 namespace Chaos
@@ -157,5 +157,5 @@ namespace Chaos
 }
 
 #if VEHICLE_DEBUGGING_ENABLED
-PRAGMA_ENABLE_OPTIMIZATION
+UE_ENABLE_OPTIMIZATION
 #endif
\ No newline at end of file
diff --git a/Engine/Source/Runtime/Experimental/ChaosVehicles/ChaosVehiclesCore/Public/SuspensionSystem.h b/Engine/Source/Runtime/Experimental/ChaosVehicles/ChaosVehiclesCore/Public/SuspensionSystem.h
--- a/Engine/Source/Runtime/Experimental/ChaosVehicles/ChaosVehiclesCore/Public/SuspensionSystem.h
+++ b/Engine/Source/Runtime/Experimental/ChaosVehicles/ChaosVehiclesCore/Public/SuspensionSystem.h
@@ -9,7 +9,7 @@
 #include "VehicleUtility.h"
 
 #if VEHICLE_DEBUGGING_ENABLED
-PRAGMA_DISABLE_OPTIMIZATION
+UE_DISABLE_OPTIMIZATION
 #endif
 
 /**
@@ -237,5 +237,5 @@ namespace Chaos
 } // namespace Chaos
 
 #if VEHICLE_DEBUGGING_ENABLED
-PRAGMA_ENABLE_OPTIMIZATION
+UE_ENABLE_OPTIMIZATION
 #endif
diff --git a/Engine/Source/Runtime/Experimental/ChaosVehicles/ChaosVehiclesCore/Public/ThrustSystem.h b/Engine/Source/Runtime/Experimental/ChaosVehicles/ChaosVehiclesCore/Public/ThrustSystem.h
--- a/Engine/Source/Runtime/Experimental/ChaosVehicles/ChaosVehiclesCore/Public/ThrustSystem.h
+++ b/Engine/Source/Runtime/Experimental/ChaosVehicles/ChaosVehiclesCore/Public/ThrustSystem.h
@@ -9,7 +9,7 @@
 #include "VehicleUtility.h"
 
 #if VEHICLE_DEBUGGING_ENABLED
-PRAGMA_DISABLE_OPTIMIZATION
+UE_DISABLE_OPTIMIZATION
 #endif
 
 namespace Chaos
@@ -103,5 +103,5 @@ namespace Chaos
 } // namespace Chaos
 
 #if VEHICLE_DEBUGGING_ENABLED
-PRAGMA_ENABLE_OPTIMIZATION
+UE_ENABLE_OPTIMIZATION
 #endif
\ No newline at end of file
diff --git a/Engine/Source/Runtime/Experimental/ChaosVehicles/ChaosVehiclesCore/Public/TireSystem.h b/Engine/Source/Runtime/Experimental/ChaosVehicles/ChaosVehiclesCore/Public/TireSystem.h
--- a/Engine/Source/Runtime/Experimental/ChaosVehicles/ChaosVehiclesCore/Public/TireSystem.h
+++ b/Engine/Source/Runtime/Experimental/ChaosVehicles/ChaosVehiclesCore/Public/TireSystem.h
@@ -6,7 +6,7 @@
 #include "VehicleUtility.h"
 
 #if VEHICLE_DEBUGGING_ENABLED
-PRAGMA_DISABLE_OPTIMIZATION
+UE_DISABLE_OPTIMIZATION
 #endif
 
 
@@ -65,5 +65,5 @@ namespace Chaos
 } // namespace Chaos
 
 #if VEHICLE_DEBUGGING_ENABLED
-PRAGMA_ENABLE_OPTIMIZATION
+UE_ENABLE_OPTIMIZATION
 #endif
diff --git a/Engine/Source/Runtime/Experimental/ChaosVehicles/ChaosVehiclesCore/Public/TransmissionSystem.h b/Engine/Source/Runtime/Experimental/ChaosVehicles/ChaosVehiclesCore/Public/TransmissionSystem.h
--- a/Engine/Source/Runtime/Experimental/ChaosVehicles/ChaosVehiclesCore/Public/TransmissionSystem.h
+++ b/Engine/Source/Runtime/Experimental/ChaosVehicles/ChaosVehiclesCore/Public/TransmissionSystem.h
@@ -8,7 +8,7 @@
 #include "VehicleUtility.h"
 
 #if VEHICLE_DEBUGGING_ENABLED
-PRAGMA_DISABLE_OPTIMIZATION
+UE_DISABLE_OPTIMIZATION
 #endif
 
 /**
@@ -237,6 +237,6 @@ namespace Chaos
 } // namespace Chaos
 
 #if VEHICLE_DEBUGGING_ENABLED
-PRAGMA_ENABLE_OPTIMIZATION
+UE_ENABLE_OPTIMIZATION
 #endif
 
diff --git a/Engine/Source/Runtime/Experimental/ChaosVehicles/ChaosVehiclesCore/Public/WheelSystem.h b/Engine/Source/Runtime/Experimental/ChaosVehicles/ChaosVehiclesCore/Public/WheelSystem.h
--- a/Engine/Source/Runtime/Experimental/ChaosVehicles/ChaosVehiclesCore/Public/WheelSystem.h
+++ b/Engine/Source/Runtime/Experimental/ChaosVehicles/ChaosVehiclesCore/Public/WheelSystem.h
@@ -10,7 +10,7 @@
 #include "VehicleUtility.h"
 
 #if VEHICLE_DEBUGGING_ENABLED
-PRAGMA_DISABLE_OPTIMIZATION
+UE_DISABLE_OPTIMIZATION
 #endif
 
 namespace Chaos
@@ -497,5 +497,5 @@ public:
 } // namespace Chaos
 
 #if VEHICLE_DEBUGGING_ENABLED
-PRAGMA_ENABLE_OPTIMIZATION
+UE_ENABLE_OPTIMIZATION
 #endif
