Chaos Vehicles: wheel collsion not updated on server due to hardcoded mode setting

diff --git a/Engine/Plugins/Experimental/ChaosVehiclesPlugin/Source/ChaosVehicles/Private/ChaosVehicleMovementComponent.cpp b/Engine/Plugins/Experimental/ChaosVehiclesPlugin/Source/ChaosVehicles/Private/ChaosVehicleMovementComponent.cpp
--- a/Engine/Plugins/Experimental/ChaosVehiclesPlugin/Source/ChaosVehicles/Private/ChaosVehicleMovementComponent.cpp
+++ b/Engine/Plugins/Experimental/ChaosVehiclesPlugin/Source/ChaosVehicles/Private/ChaosVehicleMovementComponent.cpp
@@ -780,8 +780,9 @@ void UChaosVehicleMovementComponent::OnCreatePhysicsState()
 
 	FBodyInstance* BodyInstance = nullptr;
 	if (USkeletalMeshComponent* SkeletalMesh = GetSkeletalMesh())
-	{
-		SkeletalMesh->VisibilityBasedAnimTickOption = EVisibilityBasedAnimTickOption::OnlyTickPoseWhenRendered;
+	{	
+		// this line was causing the server wheel positions to not be updated - this is already a user property so don't override it here and leave it to user to select the right option for their scenario
+		//SkeletalMesh->VisibilityBasedAnimTickOption = EVisibilityBasedAnimTickOption::OnlyTickPoseWhenRendered;
 		BodyInstance = &SkeletalMesh->BodyInstance;
 	}
 }
