[NetPhysics] Input merging for sim-proxies  who are using NetworkPhysicsComponent and RewindData

- Partial implementation for PBCM
- Full implementation for Chaos Vehicle

diff --git a/Engine/Plugins/Experimental/ChaosVehiclesPlugin/Source/ChaosVehicles/Private/ChaosVehicleManagerAsyncCallback.cpp b/Engine/Plugins/Experimental/ChaosVehiclesPlugin/Source/ChaosVehicles/Private/ChaosVehicleManagerAsyncCallback.cpp
--- a/Engine/Plugins/Experimental/ChaosVehiclesPlugin/Source/ChaosVehicles/Private/ChaosVehicleManagerAsyncCallback.cpp
+++ b/Engine/Plugins/Experimental/ChaosVehiclesPlugin/Source/ChaosVehicles/Private/ChaosVehicleManagerAsyncCallback.cpp
@@ -261,7 +261,9 @@ void FNetworkVehicleInputs::BuildDatas(const UActorComponent* NetworkComponent)
 
 void FNetworkVehicleInputs::InterpolateDatas(const FNetworkVehicleInputs& MinDatas, const FNetworkVehicleInputs& MaxDatas)
 {
-	const float LerpFactor = (LocalFrame - MinDatas.LocalFrame) / (MaxDatas.LocalFrame - MinDatas.LocalFrame);
+	const float LerpFactor = MaxDatas.LocalFrame == LocalFrame 
+	? 1.0f / (MaxDatas.LocalFrame - MinDatas.LocalFrame + 1) // Merge from min into max
+	: (LocalFrame - MinDatas.LocalFrame) / (MaxDatas.LocalFrame - MinDatas.LocalFrame); // Interpolate from min to max
 
 	TransmissionChangeTime = FMath::Lerp(MinDatas.TransmissionChangeTime, MaxDatas.TransmissionChangeTime, LerpFactor);
 	TransmissionCurrentGear = LerpFactor < 0.5 ? MinDatas.TransmissionCurrentGear : MaxDatas.TransmissionCurrentGear;
@@ -281,6 +283,15 @@ void FNetworkVehicleInputs::InterpolateDatas(const FNetworkVehicleInputs& MinDat
 	VehicleInputs.TransmissionType = LerpFactor < 0.5 ? MinDatas.VehicleInputs.TransmissionType : MaxDatas.VehicleInputs.TransmissionType;
 }
 
+void FNetworkVehicleInputs::MergeDatas(const FNetworkPhysicsDatas* FromData)
+{
+	if (const FNetworkVehicleInputs* FromDataInput = static_cast<const FNetworkVehicleInputs*>(FromData))
+	{
+		// Perform merge through InterpolateDatas
+		InterpolateDatas(*FromDataInput, *this);
+	}
+}
+
 void FNetworkVehicleInputs::DecayDatas(float DecayAmount)
 {
 	// Local adjustment to DecayAmount for vehicle implementation
diff --git a/Engine/Plugins/Experimental/ChaosVehiclesPlugin/Source/ChaosVehicles/Public/ChaosVehicleManagerAsyncCallback.h b/Engine/Plugins/Experimental/ChaosVehiclesPlugin/Source/ChaosVehicles/Public/ChaosVehicleManagerAsyncCallback.h
--- a/Engine/Plugins/Experimental/ChaosVehiclesPlugin/Source/ChaosVehicles/Public/ChaosVehicleManagerAsyncCallback.h
+++ b/Engine/Plugins/Experimental/ChaosVehiclesPlugin/Source/ChaosVehicles/Public/ChaosVehicleManagerAsyncCallback.h
@@ -136,6 +136,9 @@ struct CHAOSVEHICLES_API FNetworkVehicleInputs : public FNetworkPhysicsDatas
 
 	/** Interpolate the datas in between two inputs datas */
 	void InterpolateDatas(const FNetworkVehicleInputs& MinDatas, const FNetworkVehicleInputs& MaxDatas);
+
+	/** Merge data into this input */
+	virtual void MergeDatas(const FNetworkPhysicsDatas* FromData) override;
 };
 
 template<>
diff --git a/Engine/Source/Runtime/Engine/Private/PhysicsEngine/NetworkPhysicsComponent.cpp b/Engine/Source/Runtime/Engine/Private/PhysicsEngine/NetworkPhysicsComponent.cpp
--- a/Engine/Source/Runtime/Engine/Private/PhysicsEngine/NetworkPhysicsComponent.cpp
+++ b/Engine/Source/Runtime/Engine/Private/PhysicsEngine/NetworkPhysicsComponent.cpp
@@ -745,17 +745,27 @@ void UNetworkPhysicsComponent::OnPreProcessInputsInternal(const int32 PhysicsSte
 		if (!IsLocallyControlled() || bIsSolverResim)
 		{
 			FNetworkPhysicsDatas* PhysicsDatas = InputsDatas.Get();
+			const int32 ExpectedInputFrame = PhysicsDatas->LocalFrame + 1;
 			PhysicsDatas->LocalFrame = PhysicsStep;
+
 	#if DEBUG_NETWORK_PHYSICS
 			UE_LOG(LogChaos, Log, TEXT("		Extracting history inputs at frame %d | Component = %s"), PhysicsStep, *GetFullName());
 	#endif
 			if (InputsHistory->ExtractDatas(PhysicsStep, bIsSolverReset, PhysicsDatas))
 			{ 
 				// Calculate input decay if we are resimulating and we don't have up to date inputs
-				if (bIsSolverResim && PhysicsDatas->LocalFrame < PhysicsStep)
+				if (bIsSolverResim)
+				{
+					if (PhysicsDatas->LocalFrame < PhysicsStep)
+					{
+						const float InputDecay = GetCurrentInputDecay(PhysicsDatas);
+						PhysicsDatas->DecayDatas(InputDecay);
+					}
+				}
+				// Merge all inputs since last used input, if not resimulating
+				else if (PhysicsDatas->LocalFrame > ExpectedInputFrame)
 				{
-					const float InputDecay = GetCurrentInputDecay(PhysicsDatas);
-					PhysicsDatas->DecayDatas(InputDecay);
+					InputsHistory->MergeData(ExpectedInputFrame, PhysicsDatas);
 				}
 
 				PhysicsDatas->ApplyDatas(ActorComponent);
diff --git a/Engine/Source/Runtime/Engine/Public/Physics/NetworkPhysicsComponent.h b/Engine/Source/Runtime/Engine/Public/Physics/NetworkPhysicsComponent.h
--- a/Engine/Source/Runtime/Engine/Public/Physics/NetworkPhysicsComponent.h
+++ b/Engine/Source/Runtime/Engine/Public/Physics/NetworkPhysicsComponent.h
@@ -366,6 +366,14 @@ struct FNetworkPhysicsDatas
 	* EXAMPLE: Use to decay steering inputs to make resimulation not predict too much with a high steering value. Use DecayAmount of 0.1 to turn a steering value of 0.5 into 0.45 for example.
 	*/ 
 	virtual void DecayDatas(float DecayAmount) {}
+	
+	/** Define how to merge data together
+	* @param FromData is data from a previous frame that is getting merged into the current data.
+	* EXAMPLE: Simulated proxies might receive two inputs at the same time after having used the same input twice, to not miss any important inputs we need to take both inputs into account 
+	* and to not get behind in simulation we need to apply them both at the same simulation tick meaning we merge the two new inputs to one input.
+	*/
+	virtual void MergeDatas(const FNetworkPhysicsDatas* FromData) {}
+
 
 	friend UNetworkPhysicsComponent;
 };
diff --git a/Engine/Source/Runtime/Experimental/Chaos/Public/RewindData.h b/Engine/Source/Runtime/Experimental/Chaos/Public/RewindData.h
--- a/Engine/Source/Runtime/Experimental/Chaos/Public/RewindData.h
+++ b/Engine/Source/Runtime/Experimental/Chaos/Public/RewindData.h
@@ -44,6 +44,9 @@ struct FBaseRewindHistory
 
 	/** Extract datas from the history buffer at a given time */
 	FORCEINLINE virtual bool ExtractDatas(const int32 ExtractFrame, const bool bResetSolver, void* HistoryDatas, const bool bExactFrame = false) { return true; }
+	
+	/** Iterate over and merge data */
+	FORCEINLINE virtual void MergeData(const int32 FromFrame, void* ToData) {}
 
 	/** Record datas into the history buffer at a given time */
 	FORCEINLINE virtual bool RecordDatas(const int32 RecordFrame, const void* HistoryDatas) { return true; }
@@ -165,6 +168,19 @@ public :
 		return false;
 	}
 
+	FORCEINLINE virtual void MergeData(int32 FromFrame, void* ToData) override
+	{
+		const int32 ToFrame = static_cast<DatasType*>(ToData)->LocalFrame;
+		for (; FromFrame < ToFrame; FromFrame++)
+		{
+			const int32 LocalFrame = FromFrame % NumFrames;
+			if (FromFrame == DatasArray[LocalFrame].LocalFrame)
+			{
+				static_cast<DatasType*>(ToData)->MergeDatas(&DatasArray[LocalFrame]);
+			}
+		}
+	}
+
 	/** Load the datas from the buffer at a specific frame */
 	FORCEINLINE bool LoadDatas(const int32 LoadFrame)
 	{
