More fixes for Issue 429340

diff --git a/Engine/Plugins/Experimental/Avalanche/Source/AvalancheMedia/Public/Playable/AvaPlayableGroup.h b/Engine/Plugins/Experimental/Avalanche/Source/AvalancheMedia/Public/Playable/AvaPlayableGroup.h
--- a/Engine/Plugins/Experimental/Avalanche/Source/AvalancheMedia/Public/Playable/AvaPlayableGroup.h
+++ b/Engine/Plugins/Experimental/Avalanche/Source/AvalancheMedia/Public/Playable/AvaPlayableGroup.h
@@ -8,6 +8,7 @@
 #include "UObject/ObjectKey.h"
 #include "UObject/ObjectPtr.h"
 #include "UObject/SoftObjectPtr.h"
+#include "UObject/WeakInterfacePtr.h"
 #include "AvaPlayableGroup.generated.h"
 
 class UAvaPlayable;
diff --git a/Engine/Plugins/Experimental/ChaosVehiclesPlugin/Source/ChaosVehicles/Private/ChaosWheeledVehicleMovementComponent.cpp b/Engine/Plugins/Experimental/ChaosVehiclesPlugin/Source/ChaosVehicles/Private/ChaosWheeledVehicleMovementComponent.cpp
--- a/Engine/Plugins/Experimental/ChaosVehiclesPlugin/Source/ChaosVehicles/Private/ChaosWheeledVehicleMovementComponent.cpp
+++ b/Engine/Plugins/Experimental/ChaosVehiclesPlugin/Source/ChaosVehicles/Private/ChaosWheeledVehicleMovementComponent.cpp
@@ -1083,6 +1083,15 @@ void UChaosWheeledVehicleSimulation::FillOutputState(FChaosVehicleAsyncOutput& O
 
 }
 
+UChaosWheeledVehicleSimulation::UChaosWheeledVehicleSimulation(): bOverlapHit(false)
+{
+	QueryBox.Init();
+}
+
+UChaosWheeledVehicleSimulation::~UChaosWheeledVehicleSimulation()
+{
+}
+
 void UChaosWheeledVehicleSimulation::UpdateConstraintHandles(TArray<FPhysicsConstraintHandle>& ConstraintHandlesIn)
 {
 	UChaosVehicleSimulation::UpdateConstraintHandles(ConstraintHandlesIn);
diff --git a/Engine/Plugins/Experimental/ChaosVehiclesPlugin/Source/ChaosVehicles/Public/ChaosWheeledVehicleMovementComponent.h b/Engine/Plugins/Experimental/ChaosVehiclesPlugin/Source/ChaosVehicles/Public/ChaosWheeledVehicleMovementComponent.h
--- a/Engine/Plugins/Experimental/ChaosVehiclesPlugin/Source/ChaosVehicles/Public/ChaosWheeledVehicleMovementComponent.h
+++ b/Engine/Plugins/Experimental/ChaosVehiclesPlugin/Source/ChaosVehicles/Public/ChaosWheeledVehicleMovementComponent.h
@@ -566,15 +566,9 @@ class CHAOSVEHICLES_API UChaosWheeledVehicleSimulation : public UChaosVehicleSim
 {
 public:
 
-	UChaosWheeledVehicleSimulation()
-		: bOverlapHit(false)
-	{
-		QueryBox.Init();
-	}
+	UChaosWheeledVehicleSimulation();
 
-	virtual ~UChaosWheeledVehicleSimulation()
-	{
-	}
+	virtual ~UChaosWheeledVehicleSimulation();
 
 	virtual void Init(TUniquePtr<Chaos::FSimpleWheeledVehicle>& PVehicleIn) override
 	{
