Fix for an issue where WheelController anim node wasn't updating upstream nodes in the anim graph

diff --git a/Engine/Plugins/Experimental/ChaosVehiclesPlugin/Source/ChaosVehicles/Private/AnimNode_WheelController.cpp b/Engine/Plugins/Experimental/ChaosVehiclesPlugin/Source/ChaosVehicles/Private/AnimNode_WheelController.cpp
--- a/Engine/Plugins/Experimental/ChaosVehiclesPlugin/Source/ChaosVehicles/Private/AnimNode_WheelController.cpp
+++ b/Engine/Plugins/Experimental/ChaosVehiclesPlugin/Source/ChaosVehicles/Private/AnimNode_WheelController.cpp
@@ -138,6 +138,8 @@ void FAnimNode_WheelController::InitializeBoneReferences(const FBoneContainer& R
 
 void FAnimNode_WheelController::Initialize_AnyThread(const FAnimationInitializeContext& Context)
 {
+	FAnimNode_SkeletalControlBase::Initialize_AnyThread(Context);
+
 	AnimInstanceProxy = (FVehicleAnimationInstanceProxy*)Context.AnimInstanceProxy;	//TODO: This is cached for now because we need it in eval bone transforms.
 }
 
