Fix issue where controls not waking sleeping vehicle.

#rnx

diff --git a/Engine/Plugins/Experimental/ChaosVehiclesPlugin/Source/ChaosVehicles/Private/ChaosVehicleMovementComponent.cpp b/Engine/Plugins/Experimental/ChaosVehiclesPlugin/Source/ChaosVehicles/Private/ChaosVehicleMovementComponent.cpp
--- a/Engine/Plugins/Experimental/ChaosVehiclesPlugin/Source/ChaosVehicles/Private/ChaosVehicleMovementComponent.cpp
+++ b/Engine/Plugins/Experimental/ChaosVehiclesPlugin/Source/ChaosVehicles/Private/ChaosVehicleMovementComponent.cpp
@@ -1287,7 +1287,7 @@ void UChaosVehicleMovementComponent::ProcessSleeping(const FControlInputs& Contr
 		PrevReplicatedSteeringInput = ReplicatedState.SteeringInput;
 
 		// Wake if control input pressed
-		if ((VehicleState.bSleeping && bControlInputPressed) || GVehicleDebugParams.DisableVehicleSleep)
+		if (bControlInputPressed || GVehicleDebugParams.DisableVehicleSleep)
 		{
 			VehicleState.bSleeping = false;
 			VehicleState.SleepCounter = 0;
