New CVar to turn off suspension trace caching in Chaos Vehicle Plugin

#rnx skip

diff --git a/Engine/Plugins/Experimental/ChaosVehiclesPlugin/Source/ChaosVehicles/Private/ChaosWheeledVehicleMovementComponent.cpp b/Engine/Plugins/Experimental/ChaosVehiclesPlugin/Source/ChaosVehicles/Private/ChaosWheeledVehicleMovementComponent.cpp
--- a/Engine/Plugins/Experimental/ChaosVehiclesPlugin/Source/ChaosVehicles/Private/ChaosWheeledVehicleMovementComponent.cpp
+++ b/Engine/Plugins/Experimental/ChaosVehiclesPlugin/Source/ChaosVehicles/Private/ChaosWheeledVehicleMovementComponent.cpp
@@ -61,6 +61,7 @@ FAutoConsoleVariableRef CVarChaosVehiclesShowBatchQueryExtents(TEXT("p.Vehicle.S
 FAutoConsoleVariableRef CVarChaosVehiclesShowRaycastComponent(TEXT("p.Vehicle.ShowRaycastComponent"), GWheeledVehicleDebugParams.ShowRaycastComponent, TEXT("Enable/Disable Raycast Component Hit Visualisation."));
 FAutoConsoleVariableRef CVarChaosVehiclesShowRaycastMaterial(TEXT("p.Vehicle.ShowRaycastMaterial"), GWheeledVehicleDebugParams.ShowRaycastMaterial, TEXT("Enable/Disable Raycast Material Hit Visualisation."));
 FAutoConsoleVariableRef CVarChaosVehiclesTraceTypeOverride(TEXT("p.Vehicle.TraceTypeOverride"), GWheeledVehicleDebugParams.TraceTypeOverride, TEXT("Override ray trace type, 1=Simple, 2=Complex."));
+FAutoConsoleVariableRef CVarChaosVehiclesCacheSuspensionOffset(TEXT("p.Vehicle.CacheSuspensionOffset"), GWheeledVehicleDebugParams.CacheSuspensionOffset, TEXT("Cache Trace For Animation Optimization"));
 
 FAutoConsoleVariableRef CVarChaosVehiclesDisableSuspensionForces(TEXT("p.Vehicle.DisableSuspensionForces"), GWheeledVehicleDebugParams.DisableSuspensionForces, TEXT("Enable/Disable Suspension Forces."));
 FAutoConsoleVariableRef CVarChaosVehiclesDisableFrictionForces(TEXT("p.Vehicle.DisableFrictionForces"), GWheeledVehicleDebugParams.DisableFrictionForces, TEXT("Enable/Disable Wheel Friction Forces."));
@@ -2793,7 +2794,7 @@ float UChaosWheeledVehicleMovementComponent::GetSuspensionOffset(int WheelIndex)
 		FTransform VehicleWorldTransform = GetBodyInstance()->GetUnrealWorldTransform();
 		if (UChaosVehicleWheel* Wheel = WheelSetups[WheelIndex].WheelClass.GetDefaultObject())
 		{
-			if (WheelStatus[WheelIndex].bIsValid)
+			if (GWheeledVehicleDebugParams.CacheSuspensionOffset && WheelStatus[WheelIndex].bIsValid)
 			{
 				if (WheelStatus[WheelIndex].bInContact)
 				{
diff --git a/Engine/Plugins/Experimental/ChaosVehiclesPlugin/Source/ChaosVehicles/Public/ChaosWheeledVehicleMovementComponent.h b/Engine/Plugins/Experimental/ChaosVehiclesPlugin/Source/ChaosVehicles/Public/ChaosWheeledVehicleMovementComponent.h
--- a/Engine/Plugins/Experimental/ChaosVehiclesPlugin/Source/ChaosVehicles/Public/ChaosWheeledVehicleMovementComponent.h
+++ b/Engine/Plugins/Experimental/ChaosVehiclesPlugin/Source/ChaosVehicles/Public/ChaosWheeledVehicleMovementComponent.h
@@ -26,6 +26,7 @@ struct FWheeledVehicleDebugParams
 	bool ShowBatchQueryExtents = false;
 	bool ShowRaycastComponent = false;
 	bool ShowRaycastMaterial = false;
+	bool CacheSuspensionOffset = true;
 	int TraceTypeOverride = 0;
 
 	bool DisableSuspensionForces = false;
