Remove more uses of FBodyInstance::ActorHandle

diff --git a/Engine/Plugins/Experimental/ChaosVehiclesPlugin/Source/ChaosVehicles/Private/ChaosWheeledVehicleMovementComponent.cpp b/Engine/Plugins/Experimental/ChaosVehiclesPlugin/Source/ChaosVehicles/Private/ChaosWheeledVehicleMovementComponent.cpp
--- a/Engine/Plugins/Experimental/ChaosVehiclesPlugin/Source/ChaosVehicles/Private/ChaosWheeledVehicleMovementComponent.cpp
+++ b/Engine/Plugins/Experimental/ChaosVehiclesPlugin/Source/ChaosVehicles/Private/ChaosWheeledVehicleMovementComponent.cpp
@@ -1207,7 +1207,7 @@ void UChaosWheeledVehicleMovementComponent::FixupSkeletalMesh()
 						FBodyInstance* TargetInstance = UpdatedPrimitive->GetBodyInstance();
 						if (TargetInstance)
 						{
-							FPhysicsCommand::ExecuteWrite(TargetInstance->ActorHandle, [&](const FPhysicsActorHandle& Chassis)
+							FPhysicsCommand::ExecuteWrite(TargetInstance->GetPhysicsActorHandle(), [&](const FPhysicsActorHandle& Chassis)
 								{
 									const FVector LocalWheel = GetWheelRestingPosition(WheelSetup);
 									FPhysicsConstraintHandle ConstraintHandle = FPhysicsInterface::CreateSuspension(Chassis, LocalWheel);
@@ -1709,7 +1709,7 @@ void UChaosWheeledVehicleMovementComponent::Update(float DeltaTime)
 	{
 		if (const FBodyInstance* BodyInstance = GetBodyInstance())
 		{
-			if (auto Handle = BodyInstance->ActorHandle)
+			if (auto Handle = BodyInstance->GetPhysicsActorHandle())
 			{
 				FChaosVehicleAsyncInput* AsyncInput = static_cast<FChaosVehicleAsyncInput*>(CurAsyncInput);
 
@@ -2333,7 +2333,7 @@ void UChaosWheeledVehicleMovementComponent::SetWheelClass(int WheelIndex, TSubcl
 		if (TargetInstance && WheelIndex < Wheels.Num())
 		{
 
-			FPhysicsCommand::ExecuteWrite(TargetInstance->ActorHandle, [&](const FPhysicsActorHandle& Chassis)
+			FPhysicsCommand::ExecuteWrite(TargetInstance->GetPhysicsActorHandle(), [&](const FPhysicsActorHandle& Chassis)
 				{
 					if (VehicleSimulationPT)
 					{
@@ -2387,7 +2387,7 @@ void UChaosWheeledVehicleMovementComponent::SetSnapshot(const FWheeledSnaphotDat
 	FBodyInstance* TargetInstance = GetBodyInstance();
 	if (TargetInstance)
 	{
-		FPhysicsCommand::ExecuteWrite(TargetInstance->ActorHandle, [&](const FPhysicsActorHandle& Chassis)
+		FPhysicsCommand::ExecuteWrite(TargetInstance->GetPhysicsActorHandle(), [&](const FPhysicsActorHandle& Chassis)
 		{
 			const FWheeledSnaphotData* WheelSnapshotData = static_cast<const FWheeledSnaphotData*>(&SnapshotIn);
 
@@ -2422,7 +2422,7 @@ void UChaosWheeledVehicleMovementComponent::SetMaxEngineTorque(float Torque)
 {
 	if (FBodyInstance* TargetInstance = GetBodyInstance())
 	{
-		FPhysicsCommand::ExecuteWrite(TargetInstance->ActorHandle, [&](const FPhysicsActorHandle& Chassis)
+		FPhysicsCommand::ExecuteWrite(TargetInstance->GetPhysicsActorHandle(), [&](const FPhysicsActorHandle& Chassis)
 		{
 			if (VehicleSimulationPT && VehicleSimulationPT->PVehicle)
 			{
@@ -2436,7 +2436,7 @@ void UChaosWheeledVehicleMovementComponent::SetDragCoefficient(float DragCoeff)
 {
 	if (FBodyInstance* TargetInstance = GetBodyInstance())
 	{
-		FPhysicsCommand::ExecuteWrite(TargetInstance->ActorHandle, [&](const FPhysicsActorHandle& Chassis)
+		FPhysicsCommand::ExecuteWrite(TargetInstance->GetPhysicsActorHandle(), [&](const FPhysicsActorHandle& Chassis)
 		{
 			if (VehicleSimulationPT && VehicleSimulationPT->PVehicle)
 			{
@@ -2450,7 +2450,7 @@ void UChaosWheeledVehicleMovementComponent::SetDownforceCoefficient(float Downfo
 {
 	if (FBodyInstance* TargetInstance = GetBodyInstance())
 	{
-		FPhysicsCommand::ExecuteWrite(TargetInstance->ActorHandle, [&](const FPhysicsActorHandle& Chassis)
+		FPhysicsCommand::ExecuteWrite(TargetInstance->GetPhysicsActorHandle(), [&](const FPhysicsActorHandle& Chassis)
 		{
 			if (VehicleSimulationPT && VehicleSimulationPT->PVehicle)
 			{
@@ -2464,7 +2464,7 @@ void UChaosWheeledVehicleMovementComponent::SetDifferentialFrontRearSplit(float
 {
 	if (FBodyInstance* TargetInstance = GetBodyInstance())
 	{
-		FPhysicsCommand::ExecuteWrite(TargetInstance->ActorHandle, [&](const FPhysicsActorHandle& Chassis)
+		FPhysicsCommand::ExecuteWrite(TargetInstance->GetPhysicsActorHandle(), [&](const FPhysicsActorHandle& Chassis)
 		{
 			if (VehicleSimulationPT && VehicleSimulationPT->PVehicle)
 			{
@@ -2479,7 +2479,7 @@ void UChaosWheeledVehicleMovementComponent::SetTractionControlEnabled(int WheelI
 {
 	if (FBodyInstance* TargetInstance = GetBodyInstance())
 	{
-		FPhysicsCommand::ExecuteWrite(TargetInstance->ActorHandle, [&](const FPhysicsActorHandle& Chassis)
+		FPhysicsCommand::ExecuteWrite(TargetInstance->GetPhysicsActorHandle(), [&](const FPhysicsActorHandle& Chassis)
 		{
 			if (VehicleSimulationPT && VehicleSimulationPT->PVehicle && WheelIndex < VehicleSimulationPT->PVehicle->Wheels.Num())
 			{
@@ -2495,7 +2495,7 @@ void UChaosWheeledVehicleMovementComponent::SetABSEnabled(int WheelIndex, bool E
 {
 	if (FBodyInstance* TargetInstance = GetBodyInstance())
 	{
-		FPhysicsCommand::ExecuteWrite(TargetInstance->ActorHandle, [&](const FPhysicsActorHandle& Chassis)
+		FPhysicsCommand::ExecuteWrite(TargetInstance->GetPhysicsActorHandle(), [&](const FPhysicsActorHandle& Chassis)
 		{
 			if (VehicleSimulationPT && VehicleSimulationPT->PVehicle && WheelIndex < VehicleSimulationPT->PVehicle->Wheels.Num())
 			{
@@ -2511,7 +2511,7 @@ void UChaosWheeledVehicleMovementComponent::SetAffectedByBrake(int WheelIndex, b
 {
 	if (FBodyInstance* TargetInstance = GetBodyInstance())
 	{
-		FPhysicsCommand::ExecuteWrite(TargetInstance->ActorHandle, [&](const FPhysicsActorHandle& Chassis)
+		FPhysicsCommand::ExecuteWrite(TargetInstance->GetPhysicsActorHandle(), [&](const FPhysicsActorHandle& Chassis)
 		{
 			if (VehicleSimulationPT && VehicleSimulationPT->PVehicle && WheelIndex < VehicleSimulationPT->PVehicle->Wheels.Num())
 			{
@@ -2527,7 +2527,7 @@ void UChaosWheeledVehicleMovementComponent::SetAffectedByHandbrake(int WheelInde
 {
 	if (FBodyInstance* TargetInstance = GetBodyInstance())
 	{
-		FPhysicsCommand::ExecuteWrite(TargetInstance->ActorHandle, [&](const FPhysicsActorHandle& Chassis)
+		FPhysicsCommand::ExecuteWrite(TargetInstance->GetPhysicsActorHandle(), [&](const FPhysicsActorHandle& Chassis)
 		{
 			if (VehicleSimulationPT && VehicleSimulationPT->PVehicle && WheelIndex < VehicleSimulationPT->PVehicle->Wheels.Num())
 			{
@@ -2543,7 +2543,7 @@ void UChaosWheeledVehicleMovementComponent::SetAffectedBySteering(int WheelIndex
 {
 	if (FBodyInstance* TargetInstance = GetBodyInstance())
 	{
-		FPhysicsCommand::ExecuteWrite(TargetInstance->ActorHandle, [&](const FPhysicsActorHandle& Chassis)
+		FPhysicsCommand::ExecuteWrite(TargetInstance->GetPhysicsActorHandle(), [&](const FPhysicsActorHandle& Chassis)
 		{
 			if (VehicleSimulationPT && VehicleSimulationPT->PVehicle && WheelIndex < VehicleSimulationPT->PVehicle->Wheels.Num())
 			{
@@ -2559,7 +2559,7 @@ void UChaosWheeledVehicleMovementComponent::SetAffectedByEngine(int WheelIndex,
 {
 	if (FBodyInstance* TargetInstance = GetBodyInstance())
 	{
-		FPhysicsCommand::ExecuteWrite(TargetInstance->ActorHandle, [&](const FPhysicsActorHandle& Chassis)
+		FPhysicsCommand::ExecuteWrite(TargetInstance->GetPhysicsActorHandle(), [&](const FPhysicsActorHandle& Chassis)
 		{
 			if (VehicleSimulationPT && VehicleSimulationPT->PVehicle && WheelIndex < VehicleSimulationPT->PVehicle->Wheels.Num())
 			{
@@ -2576,7 +2576,7 @@ void UChaosWheeledVehicleMovementComponent::SetWheelRadius(int WheelIndex, float
 {
 	if (FBodyInstance* TargetInstance = GetBodyInstance())
 	{
-		FPhysicsCommand::ExecuteWrite(TargetInstance->ActorHandle, [&](const FPhysicsActorHandle& Chassis)
+		FPhysicsCommand::ExecuteWrite(TargetInstance->GetPhysicsActorHandle(), [&](const FPhysicsActorHandle& Chassis)
 		{
 			if (VehicleSimulationPT && VehicleSimulationPT->PVehicle && WheelIndex < VehicleSimulationPT->PVehicle->Wheels.Num())
 			{
@@ -2592,7 +2592,7 @@ void UChaosWheeledVehicleMovementComponent::SetWheelFrictionMultiplier(int Wheel
 {
 	if (FBodyInstance* TargetInstance = GetBodyInstance())
 	{
-		FPhysicsCommand::ExecuteWrite(TargetInstance->ActorHandle, [&](const FPhysicsActorHandle& Chassis)
+		FPhysicsCommand::ExecuteWrite(TargetInstance->GetPhysicsActorHandle(), [&](const FPhysicsActorHandle& Chassis)
 			{
 				if (VehicleSimulationPT && VehicleSimulationPT->PVehicle && WheelIndex < VehicleSimulationPT->PVehicle->Wheels.Num())
 				{
@@ -2608,7 +2608,7 @@ void UChaosWheeledVehicleMovementComponent::SetWheelSlipGraphMultiplier(int Whee
 {
 	if (FBodyInstance* TargetInstance = GetBodyInstance())
 	{
-		FPhysicsCommand::ExecuteWrite(TargetInstance->ActorHandle, [&](const FPhysicsActorHandle& Chassis)
+		FPhysicsCommand::ExecuteWrite(TargetInstance->GetPhysicsActorHandle(), [&](const FPhysicsActorHandle& Chassis)
 			{
 				if (VehicleSimulationPT && VehicleSimulationPT->PVehicle && WheelIndex < VehicleSimulationPT->PVehicle->Wheels.Num())
 				{
@@ -2624,7 +2624,7 @@ void UChaosWheeledVehicleMovementComponent::SetWheelMaxBrakeTorque(int WheelInde
 {
 	if (FBodyInstance* TargetInstance = GetBodyInstance())
 	{
-		FPhysicsCommand::ExecuteWrite(TargetInstance->ActorHandle, [&](const FPhysicsActorHandle& Chassis)
+		FPhysicsCommand::ExecuteWrite(TargetInstance->GetPhysicsActorHandle(), [&](const FPhysicsActorHandle& Chassis)
 		{
 			if (VehicleSimulationPT && VehicleSimulationPT->PVehicle && WheelIndex < VehicleSimulationPT->PVehicle->Wheels.Num())
 			{
@@ -2640,7 +2640,7 @@ void UChaosWheeledVehicleMovementComponent::SetWheelHandbrakeTorque(int WheelInd
 {
 	if (FBodyInstance* TargetInstance = GetBodyInstance())
 	{
-		FPhysicsCommand::ExecuteWrite(TargetInstance->ActorHandle, [&](const FPhysicsActorHandle& Chassis)
+		FPhysicsCommand::ExecuteWrite(TargetInstance->GetPhysicsActorHandle(), [&](const FPhysicsActorHandle& Chassis)
 		{
 			if (VehicleSimulationPT && VehicleSimulationPT->PVehicle && WheelIndex < VehicleSimulationPT->PVehicle->Wheels.Num())
 			{
@@ -2656,7 +2656,7 @@ void UChaosWheeledVehicleMovementComponent::SetWheelMaxSteerAngle(int WheelIndex
 {
 	if (FBodyInstance* TargetInstance = GetBodyInstance())
 	{
-		FPhysicsCommand::ExecuteWrite(TargetInstance->ActorHandle, [&](const FPhysicsActorHandle& Chassis)
+		FPhysicsCommand::ExecuteWrite(TargetInstance->GetPhysicsActorHandle(), [&](const FPhysicsActorHandle& Chassis)
 		{
 			if (VehicleSimulationPT && VehicleSimulationPT->PVehicle && WheelIndex < VehicleSimulationPT->PVehicle->Wheels.Num())
 			{
@@ -2672,7 +2672,7 @@ void UChaosWheeledVehicleMovementComponent::SetTorqueCombineMethod(ETorqueCombin
 {
 	if (FBodyInstance* TargetInstance = GetBodyInstance())
 	{
-		FPhysicsCommand::ExecuteWrite(TargetInstance->ActorHandle, [&](const FPhysicsActorHandle& Chassis)
+		FPhysicsCommand::ExecuteWrite(TargetInstance->GetPhysicsActorHandle(), [&](const FPhysicsActorHandle& Chassis)
 			{
 				if (VehicleSimulationPT && VehicleSimulationPT->PVehicle && WheelIndex < VehicleSimulationPT->PVehicle->Wheels.Num())
 				{
@@ -2692,7 +2692,7 @@ void UChaosWheeledVehicleMovementComponent::SetDriveTorque(float DriveTorque, in
 
 	if (FBodyInstance* TargetInstance = GetBodyInstance())
 	{
-		FPhysicsCommand::ExecuteWrite(TargetInstance->ActorHandle, [&](const FPhysicsActorHandle& Chassis)
+		FPhysicsCommand::ExecuteWrite(TargetInstance->GetPhysicsActorHandle(), [&](const FPhysicsActorHandle& Chassis)
 			{
 				if (VehicleSimulationPT && VehicleSimulationPT->PVehicle && WheelIndex < VehicleSimulationPT->PVehicle->Wheels.Num())
 				{
@@ -2710,7 +2710,7 @@ void UChaosWheeledVehicleMovementComponent::SetBrakeTorque(float BrakeTorque, in
 
 	if (FBodyInstance* TargetInstance = GetBodyInstance())
 	{
-		FPhysicsCommand::ExecuteWrite(TargetInstance->ActorHandle, [&](const FPhysicsActorHandle& Chassis)
+		FPhysicsCommand::ExecuteWrite(TargetInstance->GetPhysicsActorHandle(), [&](const FPhysicsActorHandle& Chassis)
 			{
 				if (VehicleSimulationPT && VehicleSimulationPT->PVehicle && WheelIndex < VehicleSimulationPT->PVehicle->Wheels.Num())
 				{
@@ -2728,7 +2728,7 @@ void UChaosWheeledVehicleMovementComponent::SetSuspensionParams(float Rate, floa
 
 	if (FBodyInstance* TargetInstance = GetBodyInstance())
 	{
-		FPhysicsCommand::ExecuteWrite(TargetInstance->ActorHandle, [&](const FPhysicsActorHandle& Chassis)
+		FPhysicsCommand::ExecuteWrite(TargetInstance->GetPhysicsActorHandle(), [&](const FPhysicsActorHandle& Chassis)
 			{
 				if (VehicleSimulationPT && VehicleSimulationPT->PVehicle && WheelIndex < VehicleSimulationPT->PVehicle->Wheels.Num())
 				{
@@ -2742,7 +2742,7 @@ void UChaosWheeledVehicleMovementComponent::SetSuspensionParams(float Rate, floa
 				}
 			});
 
-		FPhysicsCommand::ExecuteWrite(TargetInstance->ActorHandle, [&](const FPhysicsActorHandle& Chassis)
+		FPhysicsCommand::ExecuteWrite(TargetInstance->GetPhysicsActorHandle(), [&](const FPhysicsActorHandle& Chassis)
 			{
 				if (WheelIndex < WheelSetups.Num() && WheelIndex < Wheels.Num())
 				{
