Change to updated FBodyInstance ActorHandle API

diff --git a/Engine/Plugins/Experimental/Buoyancy/Source/Runtime/Private/BuoyancySubsystem.cpp b/Engine/Plugins/Experimental/Buoyancy/Source/Runtime/Private/BuoyancySubsystem.cpp
--- a/Engine/Plugins/Experimental/Buoyancy/Source/Runtime/Private/BuoyancySubsystem.cpp
+++ b/Engine/Plugins/Experimental/Buoyancy/Source/Runtime/Private/BuoyancySubsystem.cpp
@@ -392,7 +392,7 @@ void UBuoyancySubsystem::UpdateSplineData()
 						FBodyInstance& BodyInstance = WaterPrimitiveComponent->BodyInstance;
 						if (BodyInstance.IsValidBodyInstance())
 						{
-							if (Chaos::FSingleParticlePhysicsProxy* WaterProxy = BodyInstance.GetPhysicsActorHandle())
+							if (Chaos::FSingleParticlePhysicsProxy* WaterProxy = BodyInstance.GetPhysicsActor())
 							{
 								SplineData->SetData_GT(WaterProxy->GetGameThreadAPI(), WaterSplineData);
 							}
diff --git a/Engine/Plugins/Experimental/ChaosCaching/Source/ChaosCaching/Private/Chaos/Adapters/StaticMeshComponentCacheAdapter.cpp b/Engine/Plugins/Experimental/ChaosCaching/Source/ChaosCaching/Private/Chaos/Adapters/StaticMeshComponentCacheAdapter.cpp
--- a/Engine/Plugins/Experimental/ChaosCaching/Source/ChaosCaching/Private/Chaos/Adapters/StaticMeshComponentCacheAdapter.cpp
+++ b/Engine/Plugins/Experimental/ChaosCaching/Source/ChaosCaching/Private/Chaos/Adapters/StaticMeshComponentCacheAdapter.cpp
@@ -64,7 +64,7 @@ namespace Chaos
 	{
 		UStaticMeshComponent* MeshComp = CastChecked<UStaticMeshComponent>(InComponent);
 
-		FSingleParticlePhysicsProxy* PhysProxy = MeshComp->BodyInstance.GetPhysicsActorHandle();
+		FSingleParticlePhysicsProxy* PhysProxy = MeshComp->BodyInstance.GetPhysicsActor();
 
 		RecordToCacheInternal(PhysProxy, InRootTransform, OutFrame, InTime);
 	}
@@ -106,7 +106,7 @@ namespace Chaos
 	{
 		UStaticMeshComponent* MeshComp = CastChecked<UStaticMeshComponent>(InComponent);
 
-		FSingleParticlePhysicsProxy* PhysProxy = MeshComp->BodyInstance.GetPhysicsActorHandle();
+		FSingleParticlePhysicsProxy* PhysProxy = MeshComp->BodyInstance.GetPhysicsActor();
 
 		PlayFromCacheInternal(PhysProxy, InCache, TickRecord, OutUpdatedRigids);
 	}
@@ -179,9 +179,9 @@ namespace Chaos
 	{
 		EnsureIsInGameThreadContext();
 
-		if(Cast<UStaticMeshComponent>(InComponent) && InComponent->GetBodyInstance() && InComponent->GetBodyInstance()->GetPhysicsActorHandle())
+		if(Cast<UStaticMeshComponent>(InComponent) && InComponent->GetBodyInstance() && InComponent->GetBodyInstance()->GetPhysicsActor())
 		{
-			FPhysInterface_Chaos::SetIsKinematic_AssumesLocked(InComponent->GetBodyInstance()->GetPhysicsActorHandle(), true);
+			FPhysInterface_Chaos::SetIsKinematic_AssumesLocked(InComponent->GetBodyInstance()->GetPhysicsActor(), true);
 		}
 
 		return true;
diff --git a/Engine/Plugins/Experimental/ChaosModularVehicle/Source/ChaosModularVehicleEngine/Private/ChaosModularVehicle/ModularVehicleBaseComponent.cpp b/Engine/Plugins/Experimental/ChaosModularVehicle/Source/ChaosModularVehicleEngine/Private/ChaosModularVehicle/ModularVehicleBaseComponent.cpp
--- a/Engine/Plugins/Experimental/ChaosModularVehicle/Source/ChaosModularVehicleEngine/Private/ChaosModularVehicle/ModularVehicleBaseComponent.cpp
+++ b/Engine/Plugins/Experimental/ChaosModularVehicle/Source/ChaosModularVehicleEngine/Private/ChaosModularVehicle/ModularVehicleBaseComponent.cpp
@@ -901,7 +901,7 @@ IPhysicsProxyBase* UModularVehicleBaseComponent::GetPhysicsProxy() const
 	// SingleParticlePhysicsProxy case
 	if (const FBodyInstance* BodyInstance = UpdatedPrimitive ? UpdatedPrimitive->GetBodyInstance() : nullptr)
 	{
-		return CachedPhysicsProxy = BodyInstance->GetPhysicsActorHandle();
+		return CachedPhysicsProxy = BodyInstance->GetPhysicsActor();
 	}
 
 	return nullptr;
diff --git a/Engine/Plugins/Experimental/ChaosVehiclesPlugin/Source/ChaosVehicles/Private/ChaosVehicleMovementComponent.cpp b/Engine/Plugins/Experimental/ChaosVehiclesPlugin/Source/ChaosVehicles/Private/ChaosVehicleMovementComponent.cpp
--- a/Engine/Plugins/Experimental/ChaosVehiclesPlugin/Source/ChaosVehicles/Private/ChaosVehicleMovementComponent.cpp
+++ b/Engine/Plugins/Experimental/ChaosVehiclesPlugin/Source/ChaosVehicles/Private/ChaosVehicleMovementComponent.cpp
@@ -953,7 +953,7 @@ void UChaosVehicleMovementComponent::SetTargetGear(int32 GearNum, bool bImmediat
 
 		if (TargetInstance)
 		{
-			FPhysicsCommand::ExecuteWrite(TargetInstance->GetPhysicsActorHandle(), [&](const FPhysicsActorHandle& Chassis)
+			FPhysicsCommand::ExecuteWrite(TargetInstance->GetPhysicsActor(), [&](const FPhysicsActorHandle& Chassis)
 				{
 					if (VehicleSimulationPT && VehicleSimulationPT->PVehicle && VehicleSimulationPT->PVehicle->HasTransmission())
 					{
@@ -1502,9 +1502,9 @@ void UChaosVehicleMovementComponent::SetupVehicleMass()
 
 void UChaosVehicleMovementComponent::UpdateMassProperties(FBodyInstance* BodyInstance)
 {
-	if (BodyInstance && FPhysicsInterface::IsValid(BodyInstance->GetPhysicsActorHandle()) && FPhysicsInterface::IsRigidBody(BodyInstance->GetPhysicsActorHandle()))
+	if (BodyInstance && FPhysicsInterface::IsValid(BodyInstance->GetPhysicsActor()) && FPhysicsInterface::IsRigidBody(BodyInstance->GetPhysicsActor()))
 	{
-		FPhysicsCommand::ExecuteWrite(BodyInstance->GetPhysicsActorHandle(), [&](FPhysicsActorHandle& Actor)
+		FPhysicsCommand::ExecuteWrite(BodyInstance->GetPhysicsActor(), [&](const FPhysicsActorHandle& Actor)
 			{
 				const float MassRatio = this->Mass > 0.0f ? this->Mass / BodyInstance->GetBodyMass() : 1.0f;
 
@@ -1753,7 +1753,7 @@ void UChaosVehicleMovementComponent::Update(float DeltaTime)
 	{
 		if (const FBodyInstance* BodyInstance = GetBodyInstance())
 		{
-			if (auto Handle = BodyInstance->GetPhysicsActorHandle())
+			if (auto Handle = BodyInstance->GetPhysicsActor())
 			{
 				CurAsyncInput->Proxy = Handle;	// vehicles are never static
 				FChaosVehicleAsyncInput* AsyncInput = static_cast<FChaosVehicleAsyncInput*>(CurAsyncInput);
diff --git a/Engine/Plugins/Experimental/ChaosVehiclesPlugin/Source/ChaosVehicles/Private/ChaosWheeledVehicleMovementComponent.cpp b/Engine/Plugins/Experimental/ChaosVehiclesPlugin/Source/ChaosVehicles/Private/ChaosWheeledVehicleMovementComponent.cpp
--- a/Engine/Plugins/Experimental/ChaosVehiclesPlugin/Source/ChaosVehicles/Private/ChaosWheeledVehicleMovementComponent.cpp
+++ b/Engine/Plugins/Experimental/ChaosVehiclesPlugin/Source/ChaosVehicles/Private/ChaosWheeledVehicleMovementComponent.cpp
@@ -204,7 +204,7 @@ void UChaosWheeledVehicleSimulation::UpdateState(float DeltaTime, const FChaosVe
 			//const FHitResult& HitResult = Wheels[WheelIdx]->HitResult;
 			//if (HitResult.Component.IsValid() && HitResult.Component->GetBodyInstance())
 			//{
-			//	if (const FPhysicsActorHandle& SurfaceHandle = HitResult.Component->GetBodyInstance()->GetPhysicsActorHandle())
+			//	if (const FPhysicsActorHandle& SurfaceHandle = HitResult.Component->GetBodyInstance()->GetPhysicsActor())
 			//	{
 			//		// we are being called from the physics thread
 			//		if (Chaos::FRigidBodyHandle_Internal* SurfaceBody = SurfaceHandle->GetPhysicsThreadAPI())
@@ -1207,7 +1207,7 @@ void UChaosWheeledVehicleMovementComponent::FixupSkeletalMesh()
 						FBodyInstance* TargetInstance = UpdatedPrimitive->GetBodyInstance();
 						if (TargetInstance)
 						{
-							FPhysicsCommand::ExecuteWrite(TargetInstance->GetPhysicsActorHandle(), [&](const FPhysicsActorHandle& Chassis)
+							FPhysicsCommand::ExecuteWrite(TargetInstance->GetPhysicsActor(), [&](const FPhysicsActorHandle& Chassis)
 								{
 									const FVector LocalWheel = GetWheelRestingPosition(WheelSetup);
 									FPhysicsConstraintHandle ConstraintHandle = FPhysicsInterface::CreateSuspension(Chassis, LocalWheel);
@@ -1709,7 +1709,7 @@ void UChaosWheeledVehicleMovementComponent::Update(float DeltaTime)
 	{
 		if (const FBodyInstance* BodyInstance = GetBodyInstance())
 		{
-			if (auto Handle = BodyInstance->GetPhysicsActorHandle())
+			if (auto Handle = BodyInstance->GetPhysicsActor())
 			{
 				FChaosVehicleAsyncInput* AsyncInput = static_cast<FChaosVehicleAsyncInput*>(CurAsyncInput);
 
@@ -2333,7 +2333,7 @@ void UChaosWheeledVehicleMovementComponent::SetWheelClass(int WheelIndex, TSubcl
 		if (TargetInstance && WheelIndex < Wheels.Num())
 		{
 
-			FPhysicsCommand::ExecuteWrite(TargetInstance->GetPhysicsActorHandle(), [&](const FPhysicsActorHandle& Chassis)
+			FPhysicsCommand::ExecuteWrite(TargetInstance->GetPhysicsActor(), [&](const FPhysicsActorHandle& Chassis)
 				{
 					if (VehicleSimulationPT)
 					{
@@ -2387,7 +2387,7 @@ void UChaosWheeledVehicleMovementComponent::SetSnapshot(const FWheeledSnaphotDat
 	FBodyInstance* TargetInstance = GetBodyInstance();
 	if (TargetInstance)
 	{
-		FPhysicsCommand::ExecuteWrite(TargetInstance->GetPhysicsActorHandle(), [&](const FPhysicsActorHandle& Chassis)
+		FPhysicsCommand::ExecuteWrite(TargetInstance->GetPhysicsActor(), [&](const FPhysicsActorHandle& Chassis)
 		{
 			const FWheeledSnaphotData* WheelSnapshotData = static_cast<const FWheeledSnaphotData*>(&SnapshotIn);
 
@@ -2422,7 +2422,7 @@ void UChaosWheeledVehicleMovementComponent::SetMaxEngineTorque(float Torque)
 {
 	if (FBodyInstance* TargetInstance = GetBodyInstance())
 	{
-		FPhysicsCommand::ExecuteWrite(TargetInstance->GetPhysicsActorHandle(), [&](const FPhysicsActorHandle& Chassis)
+		FPhysicsCommand::ExecuteWrite(TargetInstance->GetPhysicsActor(), [&](const FPhysicsActorHandle& Chassis)
 		{
 			if (VehicleSimulationPT && VehicleSimulationPT->PVehicle)
 			{
@@ -2436,7 +2436,7 @@ void UChaosWheeledVehicleMovementComponent::SetDragCoefficient(float DragCoeff)
 {
 	if (FBodyInstance* TargetInstance = GetBodyInstance())
 	{
-		FPhysicsCommand::ExecuteWrite(TargetInstance->GetPhysicsActorHandle(), [&](const FPhysicsActorHandle& Chassis)
+		FPhysicsCommand::ExecuteWrite(TargetInstance->GetPhysicsActor(), [&](const FPhysicsActorHandle& Chassis)
 		{
 			if (VehicleSimulationPT && VehicleSimulationPT->PVehicle)
 			{
@@ -2450,7 +2450,7 @@ void UChaosWheeledVehicleMovementComponent::SetDownforceCoefficient(float Downfo
 {
 	if (FBodyInstance* TargetInstance = GetBodyInstance())
 	{
-		FPhysicsCommand::ExecuteWrite(TargetInstance->GetPhysicsActorHandle(), [&](const FPhysicsActorHandle& Chassis)
+		FPhysicsCommand::ExecuteWrite(TargetInstance->GetPhysicsActor(), [&](const FPhysicsActorHandle& Chassis)
 		{
 			if (VehicleSimulationPT && VehicleSimulationPT->PVehicle)
 			{
@@ -2464,7 +2464,7 @@ void UChaosWheeledVehicleMovementComponent::SetDifferentialFrontRearSplit(float
 {
 	if (FBodyInstance* TargetInstance = GetBodyInstance())
 	{
-		FPhysicsCommand::ExecuteWrite(TargetInstance->GetPhysicsActorHandle(), [&](const FPhysicsActorHandle& Chassis)
+		FPhysicsCommand::ExecuteWrite(TargetInstance->GetPhysicsActor(), [&](const FPhysicsActorHandle& Chassis)
 		{
 			if (VehicleSimulationPT && VehicleSimulationPT->PVehicle)
 			{
@@ -2479,7 +2479,7 @@ void UChaosWheeledVehicleMovementComponent::SetTractionControlEnabled(int WheelI
 {
 	if (FBodyInstance* TargetInstance = GetBodyInstance())
 	{
-		FPhysicsCommand::ExecuteWrite(TargetInstance->GetPhysicsActorHandle(), [&](const FPhysicsActorHandle& Chassis)
+		FPhysicsCommand::ExecuteWrite(TargetInstance->GetPhysicsActor(), [&](const FPhysicsActorHandle& Chassis)
 		{
 			if (VehicleSimulationPT && VehicleSimulationPT->PVehicle && WheelIndex < VehicleSimulationPT->PVehicle->Wheels.Num())
 			{
@@ -2495,7 +2495,7 @@ void UChaosWheeledVehicleMovementComponent::SetABSEnabled(int WheelIndex, bool E
 {
 	if (FBodyInstance* TargetInstance = GetBodyInstance())
 	{
-		FPhysicsCommand::ExecuteWrite(TargetInstance->GetPhysicsActorHandle(), [&](const FPhysicsActorHandle& Chassis)
+		FPhysicsCommand::ExecuteWrite(TargetInstance->GetPhysicsActor(), [&](const FPhysicsActorHandle& Chassis)
 		{
 			if (VehicleSimulationPT && VehicleSimulationPT->PVehicle && WheelIndex < VehicleSimulationPT->PVehicle->Wheels.Num())
 			{
@@ -2511,7 +2511,7 @@ void UChaosWheeledVehicleMovementComponent::SetAffectedByBrake(int WheelIndex, b
 {
 	if (FBodyInstance* TargetInstance = GetBodyInstance())
 	{
-		FPhysicsCommand::ExecuteWrite(TargetInstance->GetPhysicsActorHandle(), [&](const FPhysicsActorHandle& Chassis)
+		FPhysicsCommand::ExecuteWrite(TargetInstance->GetPhysicsActor(), [&](const FPhysicsActorHandle& Chassis)
 		{
 			if (VehicleSimulationPT && VehicleSimulationPT->PVehicle && WheelIndex < VehicleSimulationPT->PVehicle->Wheels.Num())
 			{
@@ -2527,7 +2527,7 @@ void UChaosWheeledVehicleMovementComponent::SetAffectedByHandbrake(int WheelInde
 {
 	if (FBodyInstance* TargetInstance = GetBodyInstance())
 	{
-		FPhysicsCommand::ExecuteWrite(TargetInstance->GetPhysicsActorHandle(), [&](const FPhysicsActorHandle& Chassis)
+		FPhysicsCommand::ExecuteWrite(TargetInstance->GetPhysicsActor(), [&](const FPhysicsActorHandle& Chassis)
 		{
 			if (VehicleSimulationPT && VehicleSimulationPT->PVehicle && WheelIndex < VehicleSimulationPT->PVehicle->Wheels.Num())
 			{
@@ -2543,7 +2543,7 @@ void UChaosWheeledVehicleMovementComponent::SetAffectedBySteering(int WheelIndex
 {
 	if (FBodyInstance* TargetInstance = GetBodyInstance())
 	{
-		FPhysicsCommand::ExecuteWrite(TargetInstance->GetPhysicsActorHandle(), [&](const FPhysicsActorHandle& Chassis)
+		FPhysicsCommand::ExecuteWrite(TargetInstance->GetPhysicsActor(), [&](const FPhysicsActorHandle& Chassis)
 		{
 			if (VehicleSimulationPT && VehicleSimulationPT->PVehicle && WheelIndex < VehicleSimulationPT->PVehicle->Wheels.Num())
 			{
@@ -2559,7 +2559,7 @@ void UChaosWheeledVehicleMovementComponent::SetAffectedByEngine(int WheelIndex,
 {
 	if (FBodyInstance* TargetInstance = GetBodyInstance())
 	{
-		FPhysicsCommand::ExecuteWrite(TargetInstance->GetPhysicsActorHandle(), [&](const FPhysicsActorHandle& Chassis)
+		FPhysicsCommand::ExecuteWrite(TargetInstance->GetPhysicsActor(), [&](const FPhysicsActorHandle& Chassis)
 		{
 			if (VehicleSimulationPT && VehicleSimulationPT->PVehicle && WheelIndex < VehicleSimulationPT->PVehicle->Wheels.Num())
 			{
@@ -2576,7 +2576,7 @@ void UChaosWheeledVehicleMovementComponent::SetWheelRadius(int WheelIndex, float
 {
 	if (FBodyInstance* TargetInstance = GetBodyInstance())
 	{
-		FPhysicsCommand::ExecuteWrite(TargetInstance->GetPhysicsActorHandle(), [&](const FPhysicsActorHandle& Chassis)
+		FPhysicsCommand::ExecuteWrite(TargetInstance->GetPhysicsActor(), [&](const FPhysicsActorHandle& Chassis)
 		{
 			if (VehicleSimulationPT && VehicleSimulationPT->PVehicle && WheelIndex < VehicleSimulationPT->PVehicle->Wheels.Num())
 			{
@@ -2592,7 +2592,7 @@ void UChaosWheeledVehicleMovementComponent::SetWheelFrictionMultiplier(int Wheel
 {
 	if (FBodyInstance* TargetInstance = GetBodyInstance())
 	{
-		FPhysicsCommand::ExecuteWrite(TargetInstance->GetPhysicsActorHandle(), [&](const FPhysicsActorHandle& Chassis)
+		FPhysicsCommand::ExecuteWrite(TargetInstance->GetPhysicsActor(), [&](const FPhysicsActorHandle& Chassis)
 			{
 				if (VehicleSimulationPT && VehicleSimulationPT->PVehicle && WheelIndex < VehicleSimulationPT->PVehicle->Wheels.Num())
 				{
@@ -2608,7 +2608,7 @@ void UChaosWheeledVehicleMovementComponent::SetWheelSlipGraphMultiplier(int Whee
 {
 	if (FBodyInstance* TargetInstance = GetBodyInstance())
 	{
-		FPhysicsCommand::ExecuteWrite(TargetInstance->GetPhysicsActorHandle(), [&](const FPhysicsActorHandle& Chassis)
+		FPhysicsCommand::ExecuteWrite(TargetInstance->GetPhysicsActor(), [&](const FPhysicsActorHandle& Chassis)
 			{
 				if (VehicleSimulationPT && VehicleSimulationPT->PVehicle && WheelIndex < VehicleSimulationPT->PVehicle->Wheels.Num())
 				{
@@ -2624,7 +2624,7 @@ void UChaosWheeledVehicleMovementComponent::SetWheelMaxBrakeTorque(int WheelInde
 {
 	if (FBodyInstance* TargetInstance = GetBodyInstance())
 	{
-		FPhysicsCommand::ExecuteWrite(TargetInstance->GetPhysicsActorHandle(), [&](const FPhysicsActorHandle& Chassis)
+		FPhysicsCommand::ExecuteWrite(TargetInstance->GetPhysicsActor(), [&](const FPhysicsActorHandle& Chassis)
 		{
 			if (VehicleSimulationPT && VehicleSimulationPT->PVehicle && WheelIndex < VehicleSimulationPT->PVehicle->Wheels.Num())
 			{
@@ -2640,7 +2640,7 @@ void UChaosWheeledVehicleMovementComponent::SetWheelHandbrakeTorque(int WheelInd
 {
 	if (FBodyInstance* TargetInstance = GetBodyInstance())
 	{
-		FPhysicsCommand::ExecuteWrite(TargetInstance->GetPhysicsActorHandle(), [&](const FPhysicsActorHandle& Chassis)
+		FPhysicsCommand::ExecuteWrite(TargetInstance->GetPhysicsActor(), [&](const FPhysicsActorHandle& Chassis)
 		{
 			if (VehicleSimulationPT && VehicleSimulationPT->PVehicle && WheelIndex < VehicleSimulationPT->PVehicle->Wheels.Num())
 			{
@@ -2656,7 +2656,7 @@ void UChaosWheeledVehicleMovementComponent::SetWheelMaxSteerAngle(int WheelIndex
 {
 	if (FBodyInstance* TargetInstance = GetBodyInstance())
 	{
-		FPhysicsCommand::ExecuteWrite(TargetInstance->GetPhysicsActorHandle(), [&](const FPhysicsActorHandle& Chassis)
+		FPhysicsCommand::ExecuteWrite(TargetInstance->GetPhysicsActor(), [&](const FPhysicsActorHandle& Chassis)
 		{
 			if (VehicleSimulationPT && VehicleSimulationPT->PVehicle && WheelIndex < VehicleSimulationPT->PVehicle->Wheels.Num())
 			{
@@ -2672,7 +2672,7 @@ void UChaosWheeledVehicleMovementComponent::SetTorqueCombineMethod(ETorqueCombin
 {
 	if (FBodyInstance* TargetInstance = GetBodyInstance())
 	{
-		FPhysicsCommand::ExecuteWrite(TargetInstance->GetPhysicsActorHandle(), [&](const FPhysicsActorHandle& Chassis)
+		FPhysicsCommand::ExecuteWrite(TargetInstance->GetPhysicsActor(), [&](const FPhysicsActorHandle& Chassis)
 			{
 				if (VehicleSimulationPT && VehicleSimulationPT->PVehicle && WheelIndex < VehicleSimulationPT->PVehicle->Wheels.Num())
 				{
@@ -2692,7 +2692,7 @@ void UChaosWheeledVehicleMovementComponent::SetDriveTorque(float DriveTorque, in
 
 	if (FBodyInstance* TargetInstance = GetBodyInstance())
 	{
-		FPhysicsCommand::ExecuteWrite(TargetInstance->GetPhysicsActorHandle(), [&](const FPhysicsActorHandle& Chassis)
+		FPhysicsCommand::ExecuteWrite(TargetInstance->GetPhysicsActor(), [&](const FPhysicsActorHandle& Chassis)
 			{
 				if (VehicleSimulationPT && VehicleSimulationPT->PVehicle && WheelIndex < VehicleSimulationPT->PVehicle->Wheels.Num())
 				{
@@ -2710,7 +2710,7 @@ void UChaosWheeledVehicleMovementComponent::SetBrakeTorque(float BrakeTorque, in
 
 	if (FBodyInstance* TargetInstance = GetBodyInstance())
 	{
-		FPhysicsCommand::ExecuteWrite(TargetInstance->GetPhysicsActorHandle(), [&](const FPhysicsActorHandle& Chassis)
+		FPhysicsCommand::ExecuteWrite(TargetInstance->GetPhysicsActor(), [&](const FPhysicsActorHandle& Chassis)
 			{
 				if (VehicleSimulationPT && VehicleSimulationPT->PVehicle && WheelIndex < VehicleSimulationPT->PVehicle->Wheels.Num())
 				{
@@ -2728,7 +2728,7 @@ void UChaosWheeledVehicleMovementComponent::SetSuspensionParams(float Rate, floa
 
 	if (FBodyInstance* TargetInstance = GetBodyInstance())
 	{
-		FPhysicsCommand::ExecuteWrite(TargetInstance->GetPhysicsActorHandle(), [&](const FPhysicsActorHandle& Chassis)
+		FPhysicsCommand::ExecuteWrite(TargetInstance->GetPhysicsActor(), [&](const FPhysicsActorHandle& Chassis)
 			{
 				if (VehicleSimulationPT && VehicleSimulationPT->PVehicle && WheelIndex < VehicleSimulationPT->PVehicle->Wheels.Num())
 				{
@@ -2742,7 +2742,7 @@ void UChaosWheeledVehicleMovementComponent::SetSuspensionParams(float Rate, floa
 				}
 			});
 
-		FPhysicsCommand::ExecuteWrite(TargetInstance->GetPhysicsActorHandle(), [&](const FPhysicsActorHandle& Chassis)
+		FPhysicsCommand::ExecuteWrite(TargetInstance->GetPhysicsActor(), [&](const FPhysicsActorHandle& Chassis)
 			{
 				if (WheelIndex < WheelSetups.Num() && WheelIndex < Wheels.Num())
 				{
diff --git a/Engine/Plugins/Experimental/Mover/Source/Mover/Private/Backends/MoverNetworkPhysicsLiaison.cpp b/Engine/Plugins/Experimental/Mover/Source/Mover/Private/Backends/MoverNetworkPhysicsLiaison.cpp
--- a/Engine/Plugins/Experimental/Mover/Source/Mover/Private/Backends/MoverNetworkPhysicsLiaison.cpp
+++ b/Engine/Plugins/Experimental/Mover/Source/Mover/Private/Backends/MoverNetworkPhysicsLiaison.cpp
@@ -50,7 +50,7 @@ void UMoverNetworkPhysicsLiaisonComponent::SetupConstraint()
 		const UMoverComponent& MoverComp = GetMoverComponent();
 		if (FBodyInstance* BI = MoverComp.UpdatedCompAsPrimitive ? MoverComp.UpdatedCompAsPrimitive->GetBodyInstance() : nullptr)
 		{
-			if (Chaos::FSingleParticlePhysicsProxy* CharacterProxy = BI->GetPhysicsActorHandle())
+			if (Chaos::FSingleParticlePhysicsProxy* CharacterProxy = BI->GetPhysicsActor())
 			{
 				// Create and register the constraint
 				Constraint = MakeUnique<Chaos::FCharacterGroundConstraint>();
diff --git a/Engine/Plugins/Experimental/Mover/Source/Mover/Private/Backends/MoverNetworkPhysicsLiaisonBase.cpp b/Engine/Plugins/Experimental/Mover/Source/Mover/Private/Backends/MoverNetworkPhysicsLiaisonBase.cpp
--- a/Engine/Plugins/Experimental/Mover/Source/Mover/Private/Backends/MoverNetworkPhysicsLiaisonBase.cpp
+++ b/Engine/Plugins/Experimental/Mover/Source/Mover/Private/Backends/MoverNetworkPhysicsLiaisonBase.cpp
@@ -425,7 +425,7 @@ Chaos::FUniqueIdx UMoverNetworkPhysicsLiaisonComponentBase::GetUniqueIdx() const
 	{
 		if (FBodyInstance* BI = MoverComp.UpdatedCompAsPrimitive->GetBodyInstance())
 		{
-			if (FPhysicsActorHandle ActorHandle = BI->GetPhysicsActorHandle())
+			if (FPhysicsActorHandle ActorHandle = BI->GetPhysicsActor())
 			{
 				return ActorHandle->GetGameThreadAPI().UniqueIdx();
 			}
diff --git a/Engine/Plugins/Experimental/Water/Source/Runtime/Private/BuoyancyComponent.cpp b/Engine/Plugins/Experimental/Water/Source/Runtime/Private/BuoyancyComponent.cpp
--- a/Engine/Plugins/Experimental/Water/Source/Runtime/Private/BuoyancyComponent.cpp
+++ b/Engine/Plugins/Experimental/Water/Source/Runtime/Private/BuoyancyComponent.cpp
@@ -167,7 +167,7 @@ void UBuoyancyComponent::Update(float DeltaTime)
 	{
 		if (const FBodyInstance* BodyInstance = SimulatingComponent->GetBodyInstance())
 		{
-			if (auto Handle = BodyInstance->GetPhysicsActorHandle())
+			if (auto Handle = BodyInstance->GetPhysicsActor())
 			{
 				CurAsyncInput->Proxy = Handle;
 			}
diff --git a/Engine/Plugins/Experimental/Water/Source/Runtime/Private/BuoyancyManager.cpp b/Engine/Plugins/Experimental/Water/Source/Runtime/Private/BuoyancyManager.cpp
--- a/Engine/Plugins/Experimental/Water/Source/Runtime/Private/BuoyancyManager.cpp
+++ b/Engine/Plugins/Experimental/Water/Source/Runtime/Private/BuoyancyManager.cpp
@@ -62,7 +62,7 @@ void ABuoyancyManager::OnDestroyPhysics(UActorComponent* Component)
 					{
 						if (FBodyInstance* BI = SimulatingComp->GetBodyInstance())
 						{
-							if (auto ActorHandle = BI->GetPhysicsActorHandle())
+							if (auto ActorHandle = BI->GetPhysicsActor())
 							{
 								AsyncCallback->ClearAsyncAux_External(ActorHandle->GetGameThreadAPI().UniqueIdx());
 							}
@@ -245,7 +245,7 @@ void ABuoyancyManager::InitializeAsyncAux(UBuoyancyComponent* Component)
 	{
 		if (FBodyInstance* BI = SimulatingComponent->GetBodyInstance())
 		{
-			if (auto ActorHandle = BI->GetPhysicsActorHandle())
+			if (auto ActorHandle = BI->GetPhysicsActor())
 			{
 				AsyncCallback->CreateAsyncAux_External(ActorHandle->GetGameThreadAPI().UniqueIdx(), Component->CreateAsyncAux());
 				PhysicsInitializedSimulatingComponents.Add(SimulatingComponent);
diff --git a/Engine/Plugins/PCG/Source/PCG/Private/Data/PCGCollisionWrapperData.cpp b/Engine/Plugins/PCG/Source/PCG/Private/Data/PCGCollisionWrapperData.cpp
--- a/Engine/Plugins/PCG/Source/PCG/Private/Data/PCGCollisionWrapperData.cpp
+++ b/Engine/Plugins/PCG/Source/PCG/Private/Data/PCGCollisionWrapperData.cpp
@@ -152,7 +152,7 @@ bool FPCGCollisionWrapper::GetShapeArray(FBodyInstance* Body, EPCGCollisionQuery
 	}
 
 	PhysicsInterfaceTypes::FInlineShapeArray AllShapes;
-	FillInlineShapeArray_AssumesLocked(AllShapes, Body->GetPhysicsActorHandle());
+	FillInlineShapeArray_AssumesLocked(AllShapes, Body->GetPhysicsActor());
 
 	for (auto& Shape : AllShapes)
 	{
diff --git a/Engine/Plugins/PCG/Source/PCG/Private/Data/PCGVolumeData.cpp b/Engine/Plugins/PCG/Source/PCG/Private/Data/PCGVolumeData.cpp
--- a/Engine/Plugins/PCG/Source/PCG/Private/Data/PCGVolumeData.cpp
+++ b/Engine/Plugins/PCG/Source/PCG/Private/Data/PCGVolumeData.cpp
@@ -69,7 +69,7 @@ void UPCGVolumeData::SetupVolumeBodyInstance()
 
 			// In some instances, non-collidable bodies will not be initialized, but it's not an issue for PCG so we can continue regardless.
 			// Otherwise, require that the body is not dynamic.
-			if (BodyInstance && BodySetup && (!FPhysicsInterface::IsValid(BodyInstance->GetPhysicsActorHandle()) || !BodyInstance->IsDynamic()))
+			if (BodyInstance && BodySetup && (!FPhysicsInterface::IsValid(BodyInstance->GetPhysicsActor()) || !BodyInstance->IsDynamic()))
 			{
 				ReleaseInternalBodyInstance();
 
diff --git a/Engine/Plugins/Runtime/NetworkPrediction/Source/NetworkPrediction/Private/NetworkPredictionPhysics.cpp b/Engine/Plugins/Runtime/NetworkPrediction/Source/NetworkPrediction/Private/NetworkPredictionPhysics.cpp
--- a/Engine/Plugins/Runtime/NetworkPrediction/Source/NetworkPrediction/Private/NetworkPredictionPhysics.cpp
+++ b/Engine/Plugins/Runtime/NetworkPrediction/Source/NetworkPrediction/Private/NetworkPredictionPhysics.cpp
@@ -37,7 +37,7 @@ void FNetworkPredictionPhysicsState::FinalizeInterpolatedPhysics(UPrimitiveCompo
 void FNetworkPredictionPhysicsState::NetSend(const FNetSerializeParams& P, FBodyInstance* BodyInstance)
 {
 	npCheckSlow(BodyInstance);
-	FPhysicsActorHandle& Handle = BodyInstance->GetPhysicsActorHandle();
+	FPhysicsActorHandle Handle = BodyInstance->GetPhysicsActor();
 	Chaos::FRigidBodyHandle_External& Body_External = Handle->GetGameThreadAPI();
 	if (NetworkPredictionPhysicsCvars::FullPrecision())
 	{
@@ -71,7 +71,7 @@ void FNetworkPredictionPhysicsState::NetSend(const FNetSerializeParams& P, FBody
 bool FNetworkPredictionPhysicsState::ShouldReconcile(int32 PhysicsFrame, Chaos::FRewindData* RewindData, FBodyInstance* BodyInstance, FNetworkPredictionPhysicsState* RecvState)
 {
 	npCheckSlow(BodyInstance);
-	FPhysicsActorHandle& Handle = BodyInstance->GetPhysicsActorHandle();
+	FPhysicsActorHandle Handle = BodyInstance->GetPhysicsActor();
 
 	auto CompareVector = [](const FVector& Local, const FVector& Authority, const float Tolerance, const FAnsiStringView& DebugStr)
 	{
@@ -127,7 +127,7 @@ bool FNetworkPredictionPhysicsState::ShouldReconcile(int32 PhysicsFrame, Chaos::
 
 void FNetworkPredictionPhysicsState::PerformRollback(FBodyInstance* BodyInstance, FNetworkPredictionPhysicsState* RecvState)
 {
-	FPhysicsActorHandle& Handle = BodyInstance->GetPhysicsActorHandle();
+	FPhysicsActorHandle Handle = BodyInstance->GetPhysicsActor();
 	Chaos::FRigidBodyHandle_External& Body_External = Handle->GetGameThreadAPI();
 
 	npCheckSlow(RecvState);
@@ -147,7 +147,7 @@ void FNetworkPredictionPhysicsState::PerformRollback(FBodyInstance* BodyInstance
 
 void FNetworkPredictionPhysicsState::ToString(int32 PhysicsFrame, Chaos::FRewindData* RewindData, FBodyInstance* BodyInstance, FAnsiStringBuilderBase& Builder)
 {
-	FPhysicsActorHandle& Handle = BodyInstance->GetPhysicsActorHandle();
+	FPhysicsActorHandle Handle = BodyInstance->GetPhysicsActor();
 
 	const Chaos::FGeometryParticleState LocalState = RewindData->GetPastStateAtFrame(*Handle->GetHandle_LowLevel(), PhysicsFrame);
 	ToStringInternal(LocalState.GetX(), LocalState.GetR(), LocalState.GetV(), LocalState.GetW(), Builder);
@@ -157,7 +157,7 @@ void FNetworkPredictionPhysicsState::ToString(FBodyInstance* BodyInstance, FAnsi
 {
 	if (BodyInstance)
 	{
-		FPhysicsActorHandle& Handle = BodyInstance->GetPhysicsActorHandle();
+		FPhysicsActorHandle Handle = BodyInstance->GetPhysicsActor();
 		if (Handle)
 		{
 			Chaos::FRigidBodyHandle_External& Body_External = Handle->GetGameThreadAPI();
diff --git a/Engine/Plugins/Runtime/NetworkPrediction/Source/NetworkPrediction/Private/NetworkPredictionPhysicsComponent.cpp b/Engine/Plugins/Runtime/NetworkPrediction/Source/NetworkPrediction/Private/NetworkPredictionPhysicsComponent.cpp
--- a/Engine/Plugins/Runtime/NetworkPrediction/Source/NetworkPrediction/Private/NetworkPredictionPhysicsComponent.cpp
+++ b/Engine/Plugins/Runtime/NetworkPrediction/Source/NetworkPrediction/Private/NetworkPredictionPhysicsComponent.cpp
@@ -122,6 +122,6 @@ void UNetworkPredictionPhysicsComponent::SetPrimitiveComponent(UPrimitiveCompone
 
 	if (UpdatedPrimitive)
 	{
-		PhysicsActorHandle = UpdatedPrimitive->BodyInstance.GetPhysicsActorHandle();
+		PhysicsActorHandle = UpdatedPrimitive->BodyInstance.GetPhysicsActor();
 	}
 }
diff --git a/Engine/Plugins/Runtime/NetworkPrediction/Source/NetworkPrediction/Private/NetworkPredictionUtil.cpp b/Engine/Plugins/Runtime/NetworkPrediction/Source/NetworkPrediction/Private/NetworkPredictionUtil.cpp
--- a/Engine/Plugins/Runtime/NetworkPrediction/Source/NetworkPrediction/Private/NetworkPredictionUtil.cpp
+++ b/Engine/Plugins/Runtime/NetworkPrediction/Source/NetworkPrediction/Private/NetworkPredictionUtil.cpp
@@ -15,7 +15,7 @@ namespace UE_NP
 			bool bValid = false;
 			if (UPrimitiveComponent* PrimComp = Cast<UPrimitiveComponent>(Component))
 			{
-				bValid = FPhysicsInterface::IsValid(PrimComp->BodyInstance.GetPhysicsActorHandle());
+				bValid = FPhysicsInterface::IsValid(PrimComp->BodyInstance.GetPhysicsActor());
 			}
 			return bValid;
 		};
@@ -68,7 +68,7 @@ namespace UE_NP
 
 		if (PrimitiveComponent)
 		{
-			Proxy = PrimitiveComponent->BodyInstance.GetPhysicsActorHandle();
+			Proxy = PrimitiveComponent->BodyInstance.GetPhysicsActor();
 		}
 		return Proxy;
 	}
diff --git a/Engine/Plugins/Runtime/NetworkPredictionExtras/Source/NetworkPredictionExtras/Private/MockPhysicsSimulation.cpp b/Engine/Plugins/Runtime/NetworkPredictionExtras/Source/NetworkPredictionExtras/Private/MockPhysicsSimulation.cpp
--- a/Engine/Plugins/Runtime/NetworkPredictionExtras/Source/NetworkPredictionExtras/Private/MockPhysicsSimulation.cpp
+++ b/Engine/Plugins/Runtime/NetworkPredictionExtras/Source/NetworkPredictionExtras/Private/MockPhysicsSimulation.cpp
@@ -79,7 +79,7 @@ void FMockPhysicsSimulation::SimulationTick(const FNetSimTimeStep& TimeStep, con
 						// Debug: confirm the hit component and its underlying physics body are in sync
 						if (FBodyInstance* Instance = PrimitiveComp->GetBodyInstance())
 						{
-							FPhysicsActorHandle HitHandle = Instance->GetPhysicsActorHandle();
+							FPhysicsActorHandle HitHandle = Instance->GetPhysicsActor();
 							const FVector PhysicsLocation = HitHandle->GetGameThreadAPI().X();
 
 							const FVector Delta = PhysicsLocation - PrimitiveComp->GetComponentLocation();
diff --git a/Engine/Source/Runtime/Engine/Private/Collision/CollisionConversions.cpp b/Engine/Source/Runtime/Engine/Private/Collision/CollisionConversions.cpp
--- a/Engine/Source/Runtime/Engine/Private/Collision/CollisionConversions.cpp
+++ b/Engine/Source/Runtime/Engine/Private/Collision/CollisionConversions.cpp
@@ -151,7 +151,7 @@ static void SetHitResultFromShapeAndFaceIndex(const FPhysicsShape& Shape,  const
 		}
 
 		OwningComponent = BodyInst->OwnerComponent;
-		OutResult.PhysicsObject = BodyInst->GetPhysicsActorHandle() ? BodyInst->GetPhysicsActorHandle()->GetPhysicsObject() : nullptr;
+		OutResult.PhysicsObject = BodyInst->GetPhysicsActor() ? BodyInst->GetPhysicsActor()->GetPhysicsObject() : nullptr;
 	}
 	else
 	{
@@ -672,7 +672,7 @@ void ConvertQueryOverlap(const FPhysicsShape& Shape, const FPhysicsActor& Actor,
 			OutOverlap.Component = BodyInst->OwnerComponent; // Copying weak pointer is faster than assigning raw pointer.
 			OutOverlap.ItemIndex = OwnerComponent->bMultiBodyOverlap ? BodyInst->InstanceBodyIndex : INDEX_NONE;
 		}
-		OutOverlap.PhysicsObject = BodyInst->GetPhysicsActorHandle() ? BodyInst->GetPhysicsActorHandle()->GetPhysicsObject() : nullptr;
+		OutOverlap.PhysicsObject = BodyInst->GetPhysicsActor() ? BodyInst->GetPhysicsActor()->GetPhysicsObject() : nullptr;
 
 		if (OutOverlap.PhysicsObject != nullptr)
 		{
diff --git a/Engine/Source/Runtime/Engine/Private/Components/CharacterMovementComponent.cpp b/Engine/Source/Runtime/Engine/Private/Components/CharacterMovementComponent.cpp
--- a/Engine/Source/Runtime/Engine/Private/Components/CharacterMovementComponent.cpp
+++ b/Engine/Source/Runtime/Engine/Private/Components/CharacterMovementComponent.cpp
@@ -13363,7 +13363,7 @@ void UCharacterMovementComponent::FillAsyncInput(const FVector& InputVector, FCh
 			UpdatedComponentInput->bForceGatherOverlaps = !FUpdatedComponentAsyncInput::ShouldCheckOverlapFlagToQueueOverlaps(*PrimitiveComponent);
 			UpdatedComponentInput->bGatherOverlaps = PrimitiveComponent->GetGenerateOverlapEvents() || UpdatedComponentInput->bForceGatherOverlaps;
 			UpdatedComponentInput->UpdatedComponent = PrimitiveComponent;
-			UpdatedComponentInput->PhysicsHandle = PrimitiveComponent->BodyInstance.GetPhysicsActorHandle();
+			UpdatedComponentInput->PhysicsHandle = PrimitiveComponent->BodyInstance.GetPhysicsActor();
 			ensure(UpdatedComponentInput->PhysicsHandle);
 
 			// TODO Double check we initialize all these SQ params correctly and use them in the right spots.
diff --git a/Engine/Source/Runtime/Engine/Private/Components/ShapeComponent.cpp b/Engine/Source/Runtime/Engine/Private/Components/ShapeComponent.cpp
--- a/Engine/Source/Runtime/Engine/Private/Components/ShapeComponent.cpp
+++ b/Engine/Source/Runtime/Engine/Private/Components/ShapeComponent.cpp
@@ -242,7 +242,7 @@ void UShapeComponent::CreateShapeBodySetupIfNeeded()
 		{
 			if(BodyInstance.IsValidBodyInstance())
 			{
-				FPhysicsCommand::ExecuteWrite(BodyInstance.GetActorReferenceWithWelding(), [this](const FPhysicsActorHandle& Actor)
+				FPhysicsCommand::ExecuteWrite(BodyInstance.GetPhysicsActorWeldRoot(), [this](const FPhysicsActorHandle& Actor)
 				{
 					TArray<FPhysicsShapeHandle> Shapes;
 					BodyInstance.GetAllShapes_AssumesLocked(Shapes);
diff --git a/Engine/Source/Runtime/Engine/Private/InstancedStaticMesh.cpp b/Engine/Source/Runtime/Engine/Private/InstancedStaticMesh.cpp
--- a/Engine/Source/Runtime/Engine/Private/InstancedStaticMesh.cpp
+++ b/Engine/Source/Runtime/Engine/Private/InstancedStaticMesh.cpp
@@ -231,7 +231,7 @@ class FISMExecHelper : public FSelfRegisteringExec
 						}
 						++NumValidInstanceBodies;
 						
-						FPhysicsCommand::ExecuteRead(Body->GetPhysicsActorHandle(), [&Body, &Shapes, &NumInstanceShapes](const FPhysicsActorHandle& Actor)
+						FPhysicsCommand::ExecuteRead(Body->GetPhysicsActor(), [&Body, &Shapes, &NumInstanceShapes](const FPhysicsActorHandle& Actor)
 						{
 							Shapes.Reset();
 							Body->GetAllShapes_AssumesLocked(Shapes);
@@ -2789,11 +2789,11 @@ void UInstancedStaticMeshComponent::OnAsyncDestroyPhysicsState()
 
 Chaos::FPhysicsObject* UInstancedStaticMeshComponent::GetPhysicsObjectById(Chaos::FPhysicsObjectId Id) const
 {
-	if (!InstanceBodies.IsValidIndex(Id) || !InstanceBodies[Id] || !InstanceBodies[Id]->GetPhysicsActorHandle())
+	if (!InstanceBodies.IsValidIndex(Id) || !InstanceBodies[Id] || !InstanceBodies[Id]->GetPhysicsActor())
 	{
 		return nullptr;
 	}
-	return InstanceBodies[Id]->GetPhysicsActorHandle()->GetPhysicsObject();
+	return InstanceBodies[Id]->GetPhysicsActor()->GetPhysicsObject();
 }
 
 TArray<Chaos::FPhysicsObject*> UInstancedStaticMeshComponent::GetAllPhysicsObjects() const
diff --git a/Engine/Source/Runtime/Engine/Private/PhysicsEngine/BodyInstance.cpp b/Engine/Source/Runtime/Engine/Private/PhysicsEngine/BodyInstance.cpp
--- a/Engine/Source/Runtime/Engine/Private/PhysicsEngine/BodyInstance.cpp
+++ b/Engine/Source/Runtime/Engine/Private/PhysicsEngine/BodyInstance.cpp
@@ -1571,7 +1571,7 @@ void TInitBodiesHelperBase<InAllocatorType>::InitBodies()
 				SCOPE_CYCLE_COUNTER(STAT_InitBodyAggregate);
 				for (FBodyInstance* BI : Bodies)
 				{
-					const FPhysicsActorHandle& ActorHandle = BI->GetPhysicsActor();
+					FPhysicsActorHandle ActorHandle = BI->GetPhysicsActor();
 					if (FPhysicsInterface::IsValid(ActorHandle))
 					{
 						FPhysicsInterface::AddActorToAggregate_AssumesLocked(Aggregate, ActorHandle);
diff --git a/Engine/Source/Runtime/Engine/Private/PhysicsEngine/BodySetup.cpp b/Engine/Source/Runtime/Engine/Private/PhysicsEngine/BodySetup.cpp
--- a/Engine/Source/Runtime/Engine/Private/PhysicsEngine/BodySetup.cpp
+++ b/Engine/Source/Runtime/Engine/Private/PhysicsEngine/BodySetup.cpp
@@ -722,7 +722,7 @@ void UBodySetup::AddShapesToRigidActor_AssumesLocked(
 
 	{
 		SCOPE_CYCLE_COUNTER(STAT_AddGeomToSolver);
-		FPhysicsInterface::AddGeometry(OwningInstance->GetPhysicsActorHandle(), AddParams, NewShapes);
+		FPhysicsInterface::AddGeometry(OwningInstance->GetPhysicsActor(), AddParams, NewShapes);
 	}
 }
 
diff --git a/Engine/Source/Runtime/Engine/Private/PhysicsEngine/ConstraintInstance.cpp b/Engine/Source/Runtime/Engine/Private/PhysicsEngine/ConstraintInstance.cpp
--- a/Engine/Source/Runtime/Engine/Private/PhysicsEngine/ConstraintInstance.cpp
+++ b/Engine/Source/Runtime/Engine/Private/PhysicsEngine/ConstraintInstance.cpp
@@ -455,8 +455,8 @@ float ComputeAverageMass_AssumesLocked(Chaos::FPhysicsObject* Body1, Chaos::FPhy
 
 bool GetActorRefs(FBodyInstance* Body1, FBodyInstance* Body2, FPhysicsActorHandle& OutActorRef1, FPhysicsActorHandle& OutActorRef2, UObject* DebugOwner)
 {
-	FPhysicsActorHandle ActorRef1 = Body1 ? Body1->GetPhysicsActorHandle() : FPhysicsActorHandle();
-	FPhysicsActorHandle ActorRef2 = Body2 ? Body2->GetPhysicsActorHandle() : FPhysicsActorHandle();
+	FPhysicsActorHandle ActorRef1 = Body1 ? Body1->GetPhysicsActor() : FPhysicsActorHandle();
+	FPhysicsActorHandle ActorRef2 = Body2 ? Body2->GetPhysicsActor() : FPhysicsActorHandle();
 
 	// Do not create joint unless you have two actors
 	// Do not create joint unless one of the actors is dynamic
diff --git a/Engine/Source/Runtime/Engine/Private/PhysicsEngine/EnvironmentalCollisions.cpp b/Engine/Source/Runtime/Engine/Private/PhysicsEngine/EnvironmentalCollisions.cpp
--- a/Engine/Source/Runtime/Engine/Private/PhysicsEngine/EnvironmentalCollisions.cpp
+++ b/Engine/Source/Runtime/Engine/Private/PhysicsEngine/EnvironmentalCollisions.cpp
@@ -67,7 +67,7 @@ void FEnvironmentalCollisions::AppendCollisionDataFromEnvironment(const USceneCo
 
 				bool bSuccessfulRead = false;
 
-				const FPhysicsActorHandle& ActorRef = Component->BodyInstance.GetActorReferenceWithWelding();
+				FPhysicsActorHandle ActorRef = Component->BodyInstance.GetPhysicsActorWeldRoot();
 				FPhysicsCommand::ExecuteRead(ActorRef, [&](const FPhysicsActorHandle& Actor)
 					{
 						TArray<FPhysicsShapeHandle> AllShapes;
diff --git a/Engine/Source/Runtime/Engine/Private/PhysicsEngine/Experimental/ChaosScopedSceneLock.cpp b/Engine/Source/Runtime/Engine/Private/PhysicsEngine/Experimental/ChaosScopedSceneLock.cpp
--- a/Engine/Source/Runtime/Engine/Private/PhysicsEngine/Experimental/ChaosScopedSceneLock.cpp
+++ b/Engine/Source/Runtime/Engine/Private/PhysicsEngine/Experimental/ChaosScopedSceneLock.cpp
@@ -12,7 +12,7 @@
 FScopedSceneLock_Chaos::FScopedSceneLock_Chaos(const FPhysicsActorHandle& InActorHandle, EPhysicsInterfaceScopedLockType InLockType)
 	: LockType(InLockType)
 {
-	auto Scene = GetSceneForActor(InActorHandle);
+	FChaosScene* Scene = GetSceneForActor(InActorHandle);
 	Solver = Scene ? Scene->GetSolver() : nullptr;
 	LockScene();
 }
@@ -44,7 +44,7 @@ FScopedSceneLock_Chaos::FScopedSceneLock_Chaos(const FPhysicsActorHandle& InActo
 FScopedSceneLock_Chaos::FScopedSceneLock_Chaos(FPhysicsActorHandle const* InActorHandle, EPhysicsInterfaceScopedLockType InLockType)
 	: LockType(InLockType)
 {
-	auto Scene = GetSceneForActor(*InActorHandle);
+	FChaosScene* Scene = GetSceneForActor(*InActorHandle);
 	Solver = Scene ? Scene->GetSolver() : nullptr;
 	LockScene();
 }
@@ -79,7 +79,7 @@ FScopedSceneLock_Chaos::FScopedSceneLock_Chaos(FPhysicsConstraintHandle const* I
 {
 	if (InConstraintHandle)
 	{
-		auto Scene = GetSceneForActor(InConstraintHandle);
+		FChaosScene* Scene = GetSceneForActor(InConstraintHandle);
 		Solver = Scene ? Scene->GetSolver() : nullptr;
 	}
 #if CHAOS_CHECKED
@@ -100,7 +100,7 @@ FScopedSceneLock_Chaos::FScopedSceneLock_Chaos(USkeletalMeshComponent* InSkelMes
 	{
 		for (FBodyInstance* BI : InSkelMeshComp->Bodies)
 		{
-			auto Scene = GetSceneForActor(BI->GetPhysicsActor());
+			FChaosScene* Scene = GetSceneForActor(BI->GetPhysicsActor());
 			if (Scene)
 			{
 				Solver = Scene->GetSolver();
diff --git a/Engine/Source/Runtime/Engine/Private/PhysicsEngine/Experimental/PhysScene_Chaos.cpp b/Engine/Source/Runtime/Engine/Private/PhysicsEngine/Experimental/PhysScene_Chaos.cpp
--- a/Engine/Source/Runtime/Engine/Private/PhysicsEngine/Experimental/PhysScene_Chaos.cpp
+++ b/Engine/Source/Runtime/Engine/Private/PhysicsEngine/Experimental/PhysScene_Chaos.cpp
@@ -941,7 +941,7 @@ const FBodyInstance* FPhysScene_Chaos::GetBodyInstanceFromProxyAndShape(IPhysics
 	if(ProxyInstance && InShapeIndex != INDEX_NONE)
 	{
 		PhysicsInterfaceTypes::FInlineShapeArray ShapeHandles;
-		FPhysicsInterface::GetAllShapes_AssumedLocked(ProxyInstance->GetPhysicsActorHandle(), ShapeHandles);
+		FPhysicsInterface::GetAllShapes_AssumedLocked(ProxyInstance->GetPhysicsActor(), ShapeHandles);
 
 		if(ShapeHandles.IsValidIndex(InShapeIndex))
 		{
@@ -1540,22 +1540,22 @@ void FPhysScene_Chaos::AddCustomPhysics_AssumesLocked(FBodyInstance* BodyInstanc
 
 void FPhysScene_Chaos::AddForce_AssumesLocked(FBodyInstance* BodyInstance, const FVector& Force, bool bAllowSubstepping, bool bAccelChange)
 {
-	FPhysicsInterface::AddForce_AssumesLocked(BodyInstance->GetPhysicsActorHandle(), Force, bAllowSubstepping, bAccelChange);
+	FPhysicsInterface::AddForce_AssumesLocked(BodyInstance->GetPhysicsActor(), Force, bAllowSubstepping, bAccelChange);
 }
 
 void FPhysScene_Chaos::AddForceAtPosition_AssumesLocked(FBodyInstance* BodyInstance, const FVector& Force, const FVector& Position, bool bAllowSubstepping, bool bIsLocalForce /*= false*/)
 {
-	FPhysicsInterface::AddForceAtPosition_AssumesLocked(BodyInstance->GetPhysicsActorHandle(), Force, Position, bAllowSubstepping, bIsLocalForce);
+	FPhysicsInterface::AddForceAtPosition_AssumesLocked(BodyInstance->GetPhysicsActor(), Force, Position, bAllowSubstepping, bIsLocalForce);
 }
 
 void FPhysScene_Chaos::AddRadialForceToBody_AssumesLocked(FBodyInstance* BodyInstance, const FVector& Origin, const float Radius, const float Strength, const uint8 Falloff, bool bAccelChange, bool bAllowSubstepping)
 {
-	FPhysicsInterface::AddRadialForce_AssumesLocked(BodyInstance->GetPhysicsActorHandle(), Origin, Radius, Strength, Falloff, bAccelChange, bAllowSubstepping);
+	FPhysicsInterface::AddRadialForce_AssumesLocked(BodyInstance->GetPhysicsActor(), Origin, Radius, Strength, Falloff, bAccelChange, bAllowSubstepping);
 }
 
 void FPhysScene_Chaos::ClearForces_AssumesLocked(FBodyInstance* BodyInstance, bool bAllowSubstepping)
 {
-	FPhysicsActorHandle& Handle = BodyInstance->GetPhysicsActorHandle();
+	FPhysicsActorHandle Handle = BodyInstance->GetPhysicsActor();
 	if (ensure(FPhysicsInterface::IsValid(Handle)))
 	{
 		Handle->GetGameThreadAPI().ClearForces();
@@ -1564,12 +1564,12 @@ void FPhysScene_Chaos::ClearForces_AssumesLocked(FBodyInstance* BodyInstance, bo
 
 void FPhysScene_Chaos::AddTorque_AssumesLocked(FBodyInstance* BodyInstance, const FVector& Torque, bool bAllowSubstepping, bool bAccelChange)
 {
-	FPhysicsInterface::AddTorque_AssumesLocked(BodyInstance->GetPhysicsActorHandle(), Torque, bAllowSubstepping, bAccelChange);
+	FPhysicsInterface::AddTorque_AssumesLocked(BodyInstance->GetPhysicsActor(), Torque, bAllowSubstepping, bAccelChange);
 }
 
 void FPhysScene_Chaos::ClearTorques_AssumesLocked(FBodyInstance* BodyInstance, bool bAllowSubstepping)
 {
-	FPhysicsActorHandle& Handle = BodyInstance->GetPhysicsActorHandle();
+	FPhysicsActorHandle Handle = BodyInstance->GetPhysicsActor();
 	if (ensure(FPhysicsInterface::IsValid(Handle)))
 	{
 		Handle->GetGameThreadAPI().ClearTorques();
@@ -1580,12 +1580,12 @@ void FPhysScene_Chaos::SetKinematicTarget_AssumesLocked(FBodyInstance* BodyInsta
 {
 	// #todo : Implement
 	//for now just pass it into actor directly
-	FPhysicsInterface::SetKinematicTarget_AssumesLocked(BodyInstance->GetPhysicsActorHandle(), TargetTM);
+	FPhysicsInterface::SetKinematicTarget_AssumesLocked(BodyInstance->GetPhysicsActor(), TargetTM);
 }
 
 bool FPhysScene_Chaos::GetKinematicTarget_AssumesLocked(const FBodyInstance* BodyInstance, FTransform& OutTM) const
 {
-	OutTM = FPhysicsInterface::GetKinematicTarget_AssumesLocked(BodyInstance->GetPhysicsActorHandle());
+	OutTM = FPhysicsInterface::GetKinematicTarget_AssumesLocked(BodyInstance->GetPhysicsActor());
 	return true;
 }
 
@@ -1673,7 +1673,7 @@ void GatherActorsAndTransforms(
 	for (int32 i = 0; i < NumBodies; i++)
 	{
 		FBodyInstance* BodyInst = SkelComp->Bodies[i];
-		FPhysicsActorHandle& ActorHandle = BodyInst->GetPhysicsActorHandle();
+		FPhysicsActorHandle ActorHandle = BodyInst->GetPhysicsActor();
 		if (bTeleport || !BodyInst->IsInstanceSimulatingPhysics())
 		{
 			const int32 BoneIndex = BodyInst->InstanceBoneIndex;
@@ -1940,7 +1940,7 @@ void FPhysScene_Chaos::UpdateKinematicsOnDeferredSkelMeshes()
 						continue;
 					}
 
-					FPhysicsActorHandle& ActorHandle = BodyInst->GetPhysicsActorHandle();
+					FPhysicsActorHandle ActorHandle = BodyInst->GetPhysicsActor();
 					if (GKinematicDeferralCheckValidBodies && (ActorHandle == nullptr || ActorHandle->GetSyncTimestamp() == nullptr || ActorHandle->GetMarkedDeleted()))
 					{
 						if (GKinematicDeferralLogInvalidBodies)
@@ -2025,7 +2025,7 @@ void FPhysScene_Chaos::UpdateKinematicsOnDeferredSkelMeshes()
 						continue;
 					}
 
-					FPhysicsActorHandle& ActorHandle = BodyInst->GetPhysicsActorHandle();
+					FPhysicsActorHandle ActorHandle = BodyInst->GetPhysicsActor();
 					if (GKinematicDeferralCheckValidBodies && (ActorHandle == nullptr || ActorHandle->GetSyncTimestamp() == nullptr || ActorHandle->GetMarkedDeleted()))
 					{
 						if (GKinematicDeferralLogInvalidBodies)
@@ -2358,7 +2358,7 @@ void FPhysScene_Chaos::OnSyncBodies(Chaos::FPhysicsSolverBase* Solver)
 
 						}
 
-						PhysScene->UpdateActorInAccelerationStructure(BodyInstance->GetPhysicsActorHandle());
+						PhysScene->UpdateActorInAccelerationStructure(BodyInstance->GetPhysicsActor());
 					}
 
 					if (Proxy->GetWakeEvent() != Chaos::EWakeEventEntry::None && !bPendingMove)
diff --git a/Engine/Source/Runtime/Engine/Private/PhysicsEngine/ImmediatePhysics/ImmediatePhysicsChaos/ImmediatePhysicsActorHandle_Chaos.cpp b/Engine/Source/Runtime/Engine/Private/PhysicsEngine/ImmediatePhysics/ImmediatePhysicsChaos/ImmediatePhysicsActorHandle_Chaos.cpp
--- a/Engine/Source/Runtime/Engine/Private/PhysicsEngine/ImmediatePhysics/ImmediatePhysicsChaos/ImmediatePhysicsActorHandle_Chaos.cpp
+++ b/Engine/Source/Runtime/Engine/Private/PhysicsEngine/ImmediatePhysics/ImmediatePhysicsChaos/ImmediatePhysicsActorHandle_Chaos.cpp
@@ -84,12 +84,12 @@ namespace ImmediatePhysics_Chaos
 	bool CloneGeometry(FBodyInstance* BodyInstance, EActorType ActorType, const FVector& Scale, FReal& OutMass, Chaos::FVec3& OutInertia, Chaos::FRigidTransform3& OutCoMTransform, Chaos::FImplicitObjectPtr& OutGeom, TArray<TUniquePtr<Chaos::FPerShapeData>>& OutShapes)
 	{
 		// We should only get non-simulated objects through this path, but you never know...
-		if ((BodyInstance != nullptr) && !BodyInstance->bSimulatePhysics && BodyInstance->GetPhysicsActorHandle())
+		if ((BodyInstance != nullptr) && !BodyInstance->bSimulatePhysics && BodyInstance->GetPhysicsActor())
 		{
 			OutMass = 0.0f;
 			OutInertia = FVector::ZeroVector;
 			OutCoMTransform = FTransform::Identity;
-			OutGeom = CloneGeometry(BodyInstance->GetPhysicsActorHandle()->GetGameThreadAPI().GetGeometry(), OutShapes);
+			OutGeom = CloneGeometry(BodyInstance->GetPhysicsActor()->GetGameThreadAPI().GetGeometry(), OutShapes);
 			if (OutGeom != nullptr)
 			{
 				return true;
diff --git a/Engine/Source/Runtime/Engine/Private/PhysicsEngine/PhysAnim.cpp b/Engine/Source/Runtime/Engine/Private/PhysicsEngine/PhysAnim.cpp
--- a/Engine/Source/Runtime/Engine/Private/PhysicsEngine/PhysAnim.cpp
+++ b/Engine/Source/Runtime/Engine/Private/PhysicsEngine/PhysAnim.cpp
@@ -645,7 +645,7 @@ void USkeletalMeshComponent::UpdateKinematicBonesToAnim(const TArray<FTransform>
 					{
 						continue;
 					}
-					FPhysicsActorHandle& ActorHandle = BodyInst->GetPhysicsActorHandle();
+					FPhysicsActorHandle ActorHandle = BodyInst->GetPhysicsActor();
 
 					if (FPhysicsInterface::IsValid(ActorHandle) &&
 						!FPhysicsInterface::IsStatic(ActorHandle) &&
diff --git a/Engine/Source/Runtime/Engine/Private/PhysicsEngine/PhysicalAnimationComponent.cpp b/Engine/Source/Runtime/Engine/Private/PhysicsEngine/PhysicalAnimationComponent.cpp
--- a/Engine/Source/Runtime/Engine/Private/PhysicsEngine/PhysicalAnimationComponent.cpp
+++ b/Engine/Source/Runtime/Engine/Private/PhysicsEngine/PhysicalAnimationComponent.cpp
@@ -399,7 +399,7 @@ void UPhysicalAnimationComponent::UpdatePhysicsEngineImp()
 					int32 ChildBodyIdx = PhysAsset->FindBodyIndex(PhysAnimData.BodyName);
 					if (FBodyInstance* ChildBody = (ChildBodyIdx == INDEX_NONE ? nullptr : SkeletalMeshComponent->Bodies[ChildBodyIdx]))
 					{
-						if (FPhysicsActorHandle ActorHandle = ChildBody->GetPhysicsActorHandle())
+						if (FPhysicsActorHandle ActorHandle = ChildBody->GetPhysicsActor())
 						{
 							FPhysScene* Scene = ChildBody->GetPhysicsScene();
 
@@ -431,7 +431,7 @@ void UPhysicalAnimationComponent::UpdatePhysicsEngineImp()
 							// Save reference to the kinematic actor.
 							InstanceData.TargetActor = KineActor;
 
-							ConstraintInstance->InitConstraint_AssumesLocked(ChildBody->GetPhysicsActorHandle(), InstanceData.TargetActor, 1.f);
+							ConstraintInstance->InitConstraint_AssumesLocked(ChildBody->GetPhysicsActor(), InstanceData.TargetActor, 1.f);
 						}
 					}
 				}
diff --git a/Engine/Source/Runtime/Engine/Private/PhysicsEngine/PhysicsCollisionHandler.cpp b/Engine/Source/Runtime/Engine/Private/PhysicsEngine/PhysicsCollisionHandler.cpp
--- a/Engine/Source/Runtime/Engine/Private/PhysicsEngine/PhysicsCollisionHandler.cpp
+++ b/Engine/Source/Runtime/Engine/Private/PhysicsEngine/PhysicsCollisionHandler.cpp
@@ -22,8 +22,8 @@ void UPhysicsCollisionHandler::DefaultHandleCollision_AssumesLocked(const FRigid
 
 	if(BodyInst0 && BodyInst1)
 	{
-		const FPhysicsActorHandle Handle0 = BodyInst0->GetPhysicsActorHandle();
-		const FPhysicsActorHandle Handle1 = BodyInst1->GetPhysicsActorHandle();
+		const FPhysicsActorHandle Handle0 = BodyInst0->GetPhysicsActor();
+		const FPhysicsActorHandle Handle1 = BodyInst1->GetPhysicsActor();
 
 		if(!FPhysicsInterface::IsValid(Handle0) || !FPhysicsInterface::IsValid(Handle1))
 		{
diff --git a/Engine/Source/Runtime/Engine/Private/PhysicsEngine/PhysicsConstraintComponent.cpp b/Engine/Source/Runtime/Engine/Private/PhysicsEngine/PhysicsConstraintComponent.cpp
--- a/Engine/Source/Runtime/Engine/Private/PhysicsEngine/PhysicsConstraintComponent.cpp
+++ b/Engine/Source/Runtime/Engine/Private/PhysicsEngine/PhysicsConstraintComponent.cpp
@@ -294,8 +294,8 @@ void UPhysicsConstraintComponent::InitComponentConstraint()
 	FBodyInstance* Body1 = GetBodyInstance(EConstraintFrame::Frame1);
 	FBodyInstance* Body2 = GetBodyInstance(EConstraintFrame::Frame2);
 
-	Chaos::FPhysicsObject* Object1 = (Body1 && Body1->IsValidBodyInstance()) ? Body1->GetPhysicsActorHandle()->GetPhysicsObject() : GetPhysicsObject(EConstraintFrame::Frame1);
-	Chaos::FPhysicsObject* Object2 = (Body2 && Body2->IsValidBodyInstance()) ? Body2->GetPhysicsActorHandle()->GetPhysicsObject() : GetPhysicsObject(EConstraintFrame::Frame2);
+	Chaos::FPhysicsObject* Object1 = (Body1 && Body1->IsValidBodyInstance()) ? Body1->GetPhysicsActor()->GetPhysicsObject() : GetPhysicsObject(EConstraintFrame::Frame1);
+	Chaos::FPhysicsObject* Object2 = (Body2 && Body2->IsValidBodyInstance()) ? Body2->GetPhysicsActor()->GetPhysicsObject() : GetPhysicsObject(EConstraintFrame::Frame2);
 
 	if (Object1 || Object2)
 	{
diff --git a/Engine/Source/Runtime/Engine/Private/PhysicsEngine/PhysicsHandleComponent.cpp b/Engine/Source/Runtime/Engine/Private/PhysicsEngine/PhysicsHandleComponent.cpp
--- a/Engine/Source/Runtime/Engine/Private/PhysicsEngine/PhysicsHandleComponent.cpp
+++ b/Engine/Source/Runtime/Engine/Private/PhysicsEngine/PhysicsHandleComponent.cpp
@@ -60,9 +60,9 @@ void UPhysicsHandleComponent::GrabComponent(class UPrimitiveComponent* InCompone
 
 	FRotator GrabbedRotation = FRotator::ZeroRotator;
 
-	if(FPhysicsInterface::IsValid(BodyInstance->GetPhysicsActorHandle()))
+	if(FPhysicsInterface::IsValid(BodyInstance->GetPhysicsActor()))
 	{
-		FPhysicsCommand::ExecuteRead(BodyInstance->GetPhysicsActorHandle(), [&](const FPhysicsActorHandle& Actor)
+		FPhysicsCommand::ExecuteRead(BodyInstance->GetPhysicsActor(), [&](const FPhysicsActorHandle& Actor)
 		{
 			GrabbedRotation = FPhysicsInterface::GetGlobalPose_AssumesLocked(Actor).Rotator();
 		});
@@ -108,7 +108,7 @@ void UPhysicsHandleComponent::GrabComponentImp(UPrimitiveComponent* InComponent,
 	}
 
 	// simulatable bodies should have handles.
-	FPhysicsActorHandle& InHandle = BodyInstance->GetPhysicsActorHandle();
+	FPhysicsActorHandle InHandle = BodyInstance->GetPhysicsActor();
 	if (!InHandle)
 	{
 		return;
diff --git a/Engine/Source/Runtime/Engine/Private/PhysicsEngine/PhysicsReplication.cpp b/Engine/Source/Runtime/Engine/Private/PhysicsEngine/PhysicsReplication.cpp
--- a/Engine/Source/Runtime/Engine/Private/PhysicsEngine/PhysicsReplication.cpp
+++ b/Engine/Source/Runtime/Engine/Private/PhysicsEngine/PhysicsReplication.cpp
@@ -753,7 +753,7 @@ bool FPhysicsReplication::ApplyRigidBodyState(float DeltaSeconds, FBodyInstance*
 			{
 				if (bHardsnapLegacyInPT)
 				{
-					if (Chaos::FSingleParticlePhysicsProxy* Proxy = static_cast<Chaos::FSingleParticlePhysicsProxy*>(BI->GetPhysicsActorHandle()))
+					if (Chaos::FSingleParticlePhysicsProxy* Proxy = static_cast<Chaos::FSingleParticlePhysicsProxy*>(BI->GetPhysicsActor()))
 					{
 						if (Chaos::FPBDRigidsSolver* Solver = Proxy->GetSolver<Chaos::FPBDRigidsSolver>())
 						{
@@ -805,7 +805,7 @@ bool FPhysicsReplication::ApplyRigidBodyState(float DeltaSeconds, FBodyInstance*
 				AsyncInputData.TargetState = NewState;
 				AsyncInputData.TargetState.Position = IdealWorldTM.GetLocation();
 				AsyncInputData.TargetState.Quaternion = IdealWorldTM.GetRotation();
-				AsyncInputData.Proxy = static_cast<Chaos::FSingleParticlePhysicsProxy*>(BI->GetPhysicsActorHandle());
+				AsyncInputData.Proxy = static_cast<Chaos::FSingleParticlePhysicsProxy*>(BI->GetPhysicsActor());
 				AsyncInputData.ErrorCorrection = { ErrorCorrection.LinearVelocityCoefficient, ErrorCorrection.AngularVelocityCoefficient, ErrorCorrection.PositionLerp, ErrorCorrection.AngleLerp };
 				AsyncInputData.LatencyOneWay = PingSeconds;
 
diff --git a/Engine/Source/Runtime/Engine/Private/PhysicsEngine/Tests/PhysicsBodyInstanceTests.cpp b/Engine/Source/Runtime/Engine/Private/PhysicsEngine/Tests/PhysicsBodyInstanceTests.cpp
--- a/Engine/Source/Runtime/Engine/Private/PhysicsEngine/Tests/PhysicsBodyInstanceTests.cpp
+++ b/Engine/Source/Runtime/Engine/Private/PhysicsEngine/Tests/PhysicsBodyInstanceTests.cpp
@@ -29,7 +29,7 @@ bool FPhysicsBodyInstanceTest::RunTest(const FString& Parameters)
 
 	// TODO: Physics actor handle will be changing to a ptr!
 	FBodyInstance& BI = SMC->BodyInstance;
-	FPhysicsActorHandle PH = BI.GetPhysicsActorHandle();
+	FPhysicsActorHandle PH = BI.GetPhysicsActor();
 
 	check(FPhysicsInterface::IsValid(PH));
 	check(FPhysicsInterface::IsRigidBody(PH));
diff --git a/Engine/Source/Runtime/Engine/Private/PrimitiveComponentPhysics.cpp b/Engine/Source/Runtime/Engine/Private/PrimitiveComponentPhysics.cpp
--- a/Engine/Source/Runtime/Engine/Private/PrimitiveComponentPhysics.cpp
+++ b/Engine/Source/Runtime/Engine/Private/PrimitiveComponentPhysics.cpp
@@ -1535,7 +1535,7 @@ Chaos::FPhysicsObject* UPrimitiveComponent::GetPhysicsObjectById(Chaos::FPhysics
 		return nullptr;
 	}
 
-	return BodyInstance.GetPhysicsActorHandle()->GetPhysicsObject();
+	return BodyInstance.GetPhysicsActor()->GetPhysicsObject();
 }
 
 Chaos::FPhysicsObject* UPrimitiveComponent::GetPhysicsObjectByName(const FName& Name) const
diff --git a/Engine/Source/Runtime/Engine/Private/SkeletalMeshComponentPhysics.cpp b/Engine/Source/Runtime/Engine/Private/SkeletalMeshComponentPhysics.cpp
--- a/Engine/Source/Runtime/Engine/Private/SkeletalMeshComponentPhysics.cpp
+++ b/Engine/Source/Runtime/Engine/Private/SkeletalMeshComponentPhysics.cpp
@@ -581,9 +581,9 @@ void USkeletalMeshComponent::InitCollisionRelationships()
 						{
 							if (auto* TargetBody = Bodies[TargetIndex])
 							{
-								if (FPhysicsActorHandle SourceHandle = SourceBody->GetPhysicsActorHandle())
+								if (FPhysicsActorHandle SourceHandle = SourceBody->GetPhysicsActor())
 								{
-									if (FPhysicsActorHandle TargetHandle = TargetBody->GetPhysicsActorHandle())
+									if (FPhysicsActorHandle TargetHandle = TargetBody->GetPhysicsActor())
 									{
 										if (!DisabledCollisions.Contains(SourceHandle))
 										{
@@ -622,7 +622,7 @@ void USkeletalMeshComponent::TermCollisionRelationships()
 			{
 				if (Body)
 				{
-					if (FPhysicsActorHandle Handle = Body->GetPhysicsActorHandle())
+					if (FPhysicsActorHandle Handle = Body->GetPhysicsActor())
 					{
 						CollisionRelationships.Add(Handle);
 					}
@@ -733,12 +733,12 @@ void USkeletalMeshComponent::InitArticulated(FPhysScene* PhysScene)
 	{
 		if (FBodyInstance* Body = Bodies[BodyIndex])
 		{
-			if (!Body->GetPhysicsActorHandle())
+			if (!Body->GetPhysicsActor())
 			{
 				continue;
 			}
 
-			Chaos::FPhysicsObject* PhysicsObject = Body->GetPhysicsActorHandle()->GetPhysicsObject();
+			Chaos::FPhysicsObject* PhysicsObject = Body->GetPhysicsActor()->GetPhysicsObject();
 			if (UBodySetup* Setup = Body->GetBodySetup())
 			{
 				Chaos::FPhysicsObjectInterface::SetName(PhysicsObject, Setup->BoneName);
@@ -3836,22 +3836,22 @@ FTransform USkeletalMeshComponent::GetComponentTransformFromBodyInstance(FBodyIn
 
 Chaos::FPhysicsObject* USkeletalMeshComponent::GetPhysicsObjectById(Chaos::FPhysicsObjectId Id) const
 {
-	if (!Bodies.IsValidIndex(Id) || !Bodies[Id] || !Bodies[Id]->GetPhysicsActorHandle())
+	if (!Bodies.IsValidIndex(Id) || !Bodies[Id] || !Bodies[Id]->GetPhysicsActor())
 	{
 		return nullptr;
 	}
-	return Bodies[Id]->GetPhysicsActorHandle()->GetPhysicsObject();
+	return Bodies[Id]->GetPhysicsActor()->GetPhysicsObject();
 }
 
 Chaos::FPhysicsObject* USkeletalMeshComponent::GetPhysicsObjectByName(const FName& Name) const
 {
 	FBodyInstance* Body = GetBodyInstance(Name);
-	if (!Body || !Body->GetPhysicsActorHandle())
+	if (!Body || !Body->GetPhysicsActor())
 	{
 		return nullptr;
 	}
 
-	return Body->GetPhysicsActorHandle()->GetPhysicsObject();
+	return Body->GetPhysicsActor()->GetPhysicsObject();
 }
 
 TArray<Chaos::FPhysicsObject*> USkeletalMeshComponent::GetAllPhysicsObjects() const
diff --git a/Engine/Source/Runtime/Engine/Private/Tests/AutoRTFM/AutoRTFMActorComponentTests.cpp b/Engine/Source/Runtime/Engine/Private/Tests/AutoRTFM/AutoRTFMActorComponentTests.cpp
--- a/Engine/Source/Runtime/Engine/Private/Tests/AutoRTFM/AutoRTFMActorComponentTests.cpp
+++ b/Engine/Source/Runtime/Engine/Private/Tests/AutoRTFM/AutoRTFMActorComponentTests.cpp
@@ -158,9 +158,9 @@ size_t CountAAutoRTFMTestActors()
 // See: SOL-6709
 AUTORTFM_ACTOR_COMPONENT_TEST(RegisterComponent_UnregisterComponent)
 {
-	Component->BodyInstance.SetPhysicsActorHandle(Chaos::FSingleParticlePhysicsProxy::Create(Chaos::FGeometryParticle::CreateParticle()));
-	Component->BodyInstance.GetPhysicsActorHandle()->GetParticle_LowLevel()->SetGeometry(MakeImplicitObjectPtr<Chaos::FSphere>(Chaos::FVec3(1, 2, 3), 1));
-	World->GetPhysicsScene()->GetSolver()->RegisterObject(Component->BodyInstance.GetPhysicsActorHandle());
+	Component->BodyInstance.SetPhysicsActor(Chaos::FSingleParticlePhysicsProxy::Create(Chaos::FGeometryParticle::CreateParticle()));
+	Component->BodyInstance.GetPhysicsActor()->GetParticle_LowLevel()->SetGeometry(MakeImplicitObjectPtr<Chaos::FSphere>(Chaos::FVec3(1, 2, 3), 1));
+	World->GetPhysicsScene()->GetSolver()->RegisterObject(Component->BodyInstance.GetPhysicsActor());
 
 	AutoRTFM::ETransactionResult Result = AutoRTFM::Transact([&]
 		{
diff --git a/Engine/Source/Runtime/Landscape/Private/LandscapeCollision.cpp b/Engine/Source/Runtime/Landscape/Private/LandscapeCollision.cpp
--- a/Engine/Source/Runtime/Landscape/Private/LandscapeCollision.cpp
+++ b/Engine/Source/Runtime/Landscape/Private/LandscapeCollision.cpp
@@ -497,7 +497,7 @@ void ULandscapeHeightfieldCollisionComponent::OnCreatePhysicsState()
 				// Set body instance data
 				BodyInstance.PhysicsUserData = FPhysicsUserData(&BodyInstance);
 				BodyInstance.OwnerComponent = this;
-				BodyInstance.SetPhysicsActorHandle(PhysHandle);
+				BodyInstance.SetPhysicsActor(PhysHandle);
 
 				Body_External.SetUserData(&BodyInstance.PhysicsUserData);
 
@@ -533,7 +533,7 @@ void ULandscapeHeightfieldCollisionComponent::OnDestroyPhysicsState()
 	
 	if (FPhysScene_Chaos* PhysScene = GetWorld()->GetPhysicsScene())
 	{
-		FPhysicsActorHandle& ActorHandle = BodyInstance.GetPhysicsActorHandle();
+		FPhysicsActorHandle ActorHandle = BodyInstance.GetPhysicsActor();
 		if (FPhysicsInterface::IsValid(ActorHandle))
 		{
 			PhysScene->RemoveFromComponentMaps(ActorHandle);
@@ -1898,7 +1898,7 @@ struct FMeshCollisionInitHelper
 
 		TargetInstance->PhysicsUserData = FPhysicsUserData(TargetInstance);
 		TargetInstance->OwnerComponent = Component;
-		TargetInstance->SetPhysicsActorHandle(ActorHandle);
+		TargetInstance->SetPhysicsActor(ActorHandle);
 
 		ActorHandle->GetGameThreadAPI().SetUserData(&TargetInstance->PhysicsUserData);
 	}
@@ -2053,14 +2053,14 @@ void ULandscapeHeightfieldCollisionComponent::UpdateHeightfieldRegion(int32 Comp
 			return;
 		}
 
-		if(!BodyInstance.GetPhysicsActorHandle())
+		if(!BodyInstance.GetPhysicsActor())
 		{
 			return;
 		}
 
 		// We don't lock the async scene as we only set the geometry in the sync scene's RigidActor.
 		// This function is used only during painting for line traces by the painting tools.
-		FPhysicsActorHandle PhysActorHandle = BodyInstance.GetPhysicsActorHandle();
+		FPhysicsActorHandle PhysActorHandle = BodyInstance.GetPhysicsActor();
 
 		FPhysicsCommand::ExecuteWrite(PhysActorHandle, [&](const FPhysicsActorHandle& Actor)
 		{
