Bug fix for sleep state toggling on and off when vehicle is stationary with no control inputs

diff --git a/Engine/Plugins/Experimental/ChaosVehiclesPlugin/Source/ChaosVehicles/Private/ChaosVehicleMovementComponent.cpp b/Engine/Plugins/Experimental/ChaosVehiclesPlugin/Source/ChaosVehicles/Private/ChaosVehicleMovementComponent.cpp
--- a/Engine/Plugins/Experimental/ChaosVehiclesPlugin/Source/ChaosVehicles/Private/ChaosVehicleMovementComponent.cpp
+++ b/Engine/Plugins/Experimental/ChaosVehiclesPlugin/Source/ChaosVehicles/Private/ChaosVehicleMovementComponent.cpp
@@ -1057,7 +1057,9 @@ void UChaosVehicleMovementComponent::CalcThrottleBrakeInput(float& ThrottleOut,
 			// if player isn't pressing forward or backwards...
 			if (VehicleState.ForwardSpeed < StopThreshold && VehicleState.ForwardSpeed > -StopThreshold)	//auto brake 
 			{
-				BrakeOut = 1.f;
+				// this code is causing the sleep state to toggle back off when the vehicle is stationary. It doesn't appear to be required anymore since the vehicle
+				// has another mechainsim to prevent it rolling on hills, which I believe was the initial intention of the code.
+				//BrakeOut = 1.f;
 			}
 			else
 			{
