Chaos vehicle system thruster using wrong transform function.

diff --git a/Engine/Plugins/Experimental/ChaosVehiclesPlugin/Source/ChaosVehicles/Private/ChaosVehicleMovementComponent.cpp b/Engine/Plugins/Experimental/ChaosVehiclesPlugin/Source/ChaosVehicles/Private/ChaosVehicleMovementComponent.cpp
--- a/Engine/Plugins/Experimental/ChaosVehiclesPlugin/Source/ChaosVehicles/Private/ChaosVehicleMovementComponent.cpp
+++ b/Engine/Plugins/Experimental/ChaosVehiclesPlugin/Source/ChaosVehicles/Private/ChaosVehicleMovementComponent.cpp
@@ -367,7 +367,7 @@ void UChaosVehicleSimulation::ApplyThrustForces(float DeltaTime)
 
 		Thruster.Simulate(DeltaTime);
 		FVector ThrustWorldLocation = VehicleState.VehicleWorldTransform.TransformPosition(Thruster.GetThrustLocation() + COM_Offset);
-		FVector ThrustForce = VehicleState.VehicleWorldTransform.TransformPosition(Thruster.GetThrustForce());
+		FVector ThrustForce = VehicleState.VehicleWorldTransform.TransformVector(Thruster.GetThrustForce());
 
 		AddForceAtPosition(ThrustForce, ThrustWorldLocation);
 	}
