Change a log warning to a low priority log

diff --git a/Engine/Plugins/Experimental/ChaosVehiclesPlugin/Source/ChaosVehicles/Private/ChaosWheeledVehicleMovementComponent.cpp b/Engine/Plugins/Experimental/ChaosVehiclesPlugin/Source/ChaosVehicles/Private/ChaosWheeledVehicleMovementComponent.cpp
--- a/Engine/Plugins/Experimental/ChaosVehiclesPlugin/Source/ChaosVehicles/Private/ChaosWheeledVehicleMovementComponent.cpp
+++ b/Engine/Plugins/Experimental/ChaosVehiclesPlugin/Source/ChaosVehicles/Private/ChaosWheeledVehicleMovementComponent.cpp
@@ -1491,7 +1491,7 @@ void UChaosWheeledVehicleMovementComponent::SetupVehicle(TUniquePtr<Chaos::FSimp
 		{
 			ActorName = GetOwner()->GetName();
 		}
-		UE_LOG(LogVehicle, Warning, TEXT("Vehicle %s has no torque curve defined, disabling mechanical simulation."), *ActorName);
+		UE_LOG(LogVehicle, Log, TEXT("Vehicle %s has no torque curve defined, disabling mechanical simulation."), *ActorName);
 
 		bMechanicalSimEnabled = false;
 	}
