Fix bug where vehicle gets stopped by small step when moving slowly, wrong normal being used

diff --git a/Engine/Plugins/Experimental/ChaosVehiclesPlugin/Source/ChaosVehicles/Private/ChaosWheeledVehicleMovementComponent.cpp b/Engine/Plugins/Experimental/ChaosVehiclesPlugin/Source/ChaosVehicles/Private/ChaosWheeledVehicleMovementComponent.cpp
--- a/Engine/Plugins/Experimental/ChaosVehiclesPlugin/Source/ChaosVehicles/Private/ChaosWheeledVehicleMovementComponent.cpp
+++ b/Engine/Plugins/Experimental/ChaosVehiclesPlugin/Source/ChaosVehicles/Private/ChaosWheeledVehicleMovementComponent.cpp
@@ -620,7 +620,7 @@ void UChaosWheeledVehicleSimulation::ApplySuspensionForces(float DeltaTime, TArr
 
 							Chaos::FPhysicsSolver* Solver = Proxy->GetSolver<Chaos::FPhysicsSolver>();
 
-							Solver->SetSuspensionTarget(Constraint, TargetPos, HitResult.ImpactNormal, PWheel.InContact());
+							Solver->SetSuspensionTarget(Constraint, TargetPos, HitResult.Normal, PWheel.InContact());
 						}
 					}
 				}
