Chaos vehicles plugin : Fix shutdown crash bug

diff --git a/Engine/Plugins/Experimental/ChaosVehiclesPlugin/Source/ChaosVehicles/Private/ChaosVehicleManagerAsyncCallback.cpp b/Engine/Plugins/Experimental/ChaosVehiclesPlugin/Source/ChaosVehicles/Private/ChaosVehicleManagerAsyncCallback.cpp
--- a/Engine/Plugins/Experimental/ChaosVehiclesPlugin/Source/ChaosVehicles/Private/ChaosVehicleManagerAsyncCallback.cpp
+++ b/Engine/Plugins/Experimental/ChaosVehiclesPlugin/Source/ChaosVehicles/Private/ChaosVehicleManagerAsyncCallback.cpp
@@ -38,11 +38,14 @@ void FChaosVehicleManagerAsyncCallback::ProcessInputs_Internal(int32 PhysicsStep
 			continue;
 		}
 		bool bIsResimming = false;
-		if (FPhysScene* PhysScene = VehicleInput->Vehicle->GetWorld()->GetPhysicsScene())
+		if (UWorld* World = VehicleInput->Vehicle->GetWorld())
 		{
-			if (Chaos::FPhysicsSolver* LocalSolver = PhysScene->GetSolver())
+			if (FPhysScene* PhysScene = World->GetPhysicsScene())
 			{
-				bIsResimming = LocalSolver->GetEvolution()->IsResimming();
+				if (Chaos::FPhysicsSolver* LocalSolver = PhysScene->GetSolver())
+				{
+					bIsResimming = LocalSolver->GetEvolution()->IsResimming();
+				}
 			}
 		}
 
