Suppress divide by 0 compiler SA error in UChaosWheeledVehicleMovementComponent::DrawDebug

Some of my changes affect inlining which is causing this to trigger. I'm just adding an assert which also hints to the compiler that it can't be 0.

#rnx

diff --git a/Engine/Plugins/Experimental/ChaosVehiclesPlugin/Source/ChaosVehicles/Private/ChaosWheeledVehicleMovementComponent.cpp b/Engine/Plugins/Experimental/ChaosVehiclesPlugin/Source/ChaosVehicles/Private/ChaosWheeledVehicleMovementComponent.cpp
--- a/Engine/Plugins/Experimental/ChaosVehiclesPlugin/Source/ChaosVehicles/Private/ChaosWheeledVehicleMovementComponent.cpp
+++ b/Engine/Plugins/Experimental/ChaosVehiclesPlugin/Source/ChaosVehicles/Private/ChaosWheeledVehicleMovementComponent.cpp
@@ -2074,9 +2074,12 @@ void UChaosWheeledVehicleMovementComponent::DrawDebug(UCanvas* Canvas, float& YL
 
 		for (int Gear = 1; Gear <= TSetup.ForwardGearRatios.Num(); Gear++)
 		{
+			const float GearRatio = TSetup.GetGearRatio(Gear);
+			check(GearRatio != 0.f);
+
 			for (int EngineRPM = 0; EngineRPM <= ESetup.MaxRPM; EngineRPM += 10)
 			{
-				float RPMOut = EngineRPM / TSetup.GetGearRatio(Gear);
+				float RPMOut = EngineRPM / GearRatio;
 
 				float X = RPMOut / (ESetup.MaxRPM / LongGearRatio);
 				float Y = ESetup.GetTorqueFromRPM(EngineRPM) * TSetup.GetGearRatio(Gear) / (MaxTorque*MaxGearRatio);
