Dstm: Chaos vehicle support 1.3: Add migration support for VehicleMovementComponent. Add custom serialization for VehicleMovementComponent::PVehicleOutput.

#Synced-CL 48437998

diff --git a/Engine/Plugins/Experimental/ChaosVehiclesPlugin/Source/ChaosVehicles/Private/ChaosVehicleMovementComponent.cpp b/Engine/Plugins/Experimental/ChaosVehiclesPlugin/Source/ChaosVehicles/Private/ChaosVehicleMovementComponent.cpp
--- a/Engine/Plugins/Experimental/ChaosVehiclesPlugin/Source/ChaosVehicles/Private/ChaosVehicleMovementComponent.cpp
+++ b/Engine/Plugins/Experimental/ChaosVehiclesPlugin/Source/ChaosVehicles/Private/ChaosVehicleMovementComponent.cpp
@@ -667,8 +667,33 @@ void UChaosVehicleMovementComponent::Serialize(FArchive& Ar)
 	Super::Serialize(Ar);
 
 	// Custom serialization goes here...
+#if UE_WITH_REMOTE_OBJECT_HANDLE
+	if (Ar.IsMigratingRemoteObjects())
+	{
+		MigratePhysicsState(Ar);
+	}
+#endif
+}
+
+#if UE_WITH_REMOTE_OBJECT_HANDLE
+void UChaosVehicleMovementComponent::MigratePhysicsState(FArchive& Ar)
+{
+	const FUObjectMigrationContext* MigrationContext = Ar.GetMigrationContext();
+	check(MigrationContext);
+
+	if (Ar.IsLoading())
+	{
+		PVehicleOutput = MakeUnique<FPhysicsVehicleOutput>();
+	}
+
+	if (PVehicleOutput.IsValid())
+	{
+		FPhysicsVehicleOutput::StaticStruct()->SerializeItem(Ar, PVehicleOutput.Get(), nullptr);
+	}
 }
 
+#endif 
+
 #if WITH_EDITOR
 void UChaosVehicleMovementComponent::PostEditChangeProperty(FPropertyChangedEvent& PropertyChangedEvent)
 {
diff --git a/Engine/Plugins/Experimental/ChaosVehiclesPlugin/Source/ChaosVehicles/Public/ChaosVehicleManagerAsyncCallback.h b/Engine/Plugins/Experimental/ChaosVehiclesPlugin/Source/ChaosVehicles/Public/ChaosVehicleManagerAsyncCallback.h
--- a/Engine/Plugins/Experimental/ChaosVehiclesPlugin/Source/ChaosVehicles/Public/ChaosVehicleManagerAsyncCallback.h
+++ b/Engine/Plugins/Experimental/ChaosVehiclesPlugin/Source/ChaosVehicles/Public/ChaosVehicleManagerAsyncCallback.h
@@ -221,8 +221,11 @@ struct TStructOpsTypeTraits<FNetworkVehicleStates> : public TStructOpsTypeTraits
 };
 
 /** Local physics wheels outputs not replicated across the network */
+USTRUCT()
 struct CHAOSVEHICLES_API FWheelsOutput
 {
+	GENERATED_BODY()
+
 	FWheelsOutput()
 		: InContact(false)
 		, SteeringAngle(0.f)
@@ -250,41 +253,85 @@ struct CHAOSVEHICLES_API FWheelsOutput
 	}
 
 	// wheels
+	UPROPERTY()
 	bool InContact;
+
+	UPROPERTY()
 	float SteeringAngle;
+
+	UPROPERTY()
 	float AngularPosition;
+
+	UPROPERTY()
 	float AngularVelocity;
+
+	UPROPERTY()
 	float WheelRadius;
 
+	UPROPERTY()
 	float LateralAdhesiveLimit;
+
+	UPROPERTY()
 	float LongitudinalAdhesiveLimit;
 
+	UPROPERTY()
 	float SlipAngle;
+
+	UPROPERTY()
 	bool bIsSlipping;
+
+	UPROPERTY()
 	float SlipMagnitude;
+
+	UPROPERTY()
 	bool bIsSkidding;
+
+	UPROPERTY()
 	float SkidMagnitude;
+
+	UPROPERTY()
 	FVector SkidNormal;
 
 	// suspension related
+	UPROPERTY()
 	float SuspensionOffset;
+
+	UPROPERTY()
 	float SpringForce;
+
+	UPROPERTY()
 	float NormalizedSuspensionLength;
 
+	UPROPERTY()
 	float DriveTorque;
+
+	UPROPERTY()
 	float BrakeTorque;
+
+	UPROPERTY()
 	bool bABSActivated;
+
+	UPROPERTY()
 	bool bBlockingHit;
+
+	UPROPERTY()
 	FVector ImpactPoint;
+
+	UPROPERTY()
 	FVector HitLocation;
+
+	UPROPERTY()
 	TWeakObjectPtr<UPhysicalMaterial> PhysMaterial;
 };
 
 /**
  * Per Vehicle Output State from Physics Thread to Game Thread                            
  */
+USTRUCT()
 struct CHAOSVEHICLES_API FPhysicsVehicleOutput
 {
+	GENERATED_BODY()
+
 	FPhysicsVehicleOutput()
 		: CurrentGear(0)
 		, TargetGear(0)
@@ -295,12 +342,25 @@ struct CHAOSVEHICLES_API FPhysicsVehicleOutput
 	{
 	}
 
+	UPROPERTY()
 	TArray<FWheelsOutput> Wheels;
+
+	UPROPERTY()
 	int32 CurrentGear;
+
+	UPROPERTY()
 	int32 TargetGear;
+
+	UPROPERTY()
 	float EngineRPM;
+
+	UPROPERTY()
 	float EngineTorque;
+
+	UPROPERTY()
 	float TransmissionRPM;
+
+	UPROPERTY()
 	float TransmissionTorque;
 };
 
diff --git a/Engine/Plugins/Experimental/ChaosVehiclesPlugin/Source/ChaosVehicles/Public/ChaosVehicleMovementComponent.h b/Engine/Plugins/Experimental/ChaosVehiclesPlugin/Source/ChaosVehicles/Public/ChaosVehicleMovementComponent.h
--- a/Engine/Plugins/Experimental/ChaosVehiclesPlugin/Source/ChaosVehicles/Public/ChaosVehicleMovementComponent.h
+++ b/Engine/Plugins/Experimental/ChaosVehiclesPlugin/Source/ChaosVehicles/Public/ChaosVehicleMovementComponent.h
@@ -1261,6 +1261,11 @@ protected:
 	// draw 2D debug line to UI canvas
 	void DrawLine2D(UCanvas* Canvas, const FVector2D& StartPos, const FVector2D& EndPos, FColor Color, float Thickness = 1.f);
 
+#if UE_WITH_REMOTE_OBJECT_HANDLE
+	// migrates physics related state (such as vehicle ouput)
+	ENGINE_API virtual void MigratePhysicsState(FArchive& Ar);
+#endif
+
 	float GetForwardAcceleration()
 	{
 		return VehicleState.ForwardsAcceleration;
