OG Vehicle plugin, fix a crash where client network data is not available when game is first starting up.

#Synced-CL 48355415

diff --git a/Engine/Plugins/Experimental/ChaosVehiclesPlugin/Source/ChaosVehicles/Private/ChaosVehicleManagerAsyncCallback.cpp b/Engine/Plugins/Experimental/ChaosVehiclesPlugin/Source/ChaosVehicles/Private/ChaosVehicleManagerAsyncCallback.cpp
--- a/Engine/Plugins/Experimental/ChaosVehiclesPlugin/Source/ChaosVehicles/Private/ChaosVehicleManagerAsyncCallback.cpp
+++ b/Engine/Plugins/Experimental/ChaosVehiclesPlugin/Source/ChaosVehicles/Private/ChaosVehicleManagerAsyncCallback.cpp
@@ -327,9 +327,12 @@ void FNetworkVehicleStates::ApplyData(UActorComponent* NetworkComponent) const
 		if (TUniquePtr<Chaos::FSimpleWheeledVehicle>& Vehicle = VehicleSimulation->PVehicle)
 		{
 			Vehicle->GetEngine().SetEngineOmega(EngineOmega);
-		
+
+			// get the size of the source data - all arrays are always going to be the same length		
+			const int32 NumWheels = FMath::Min(SuspensionLastSpringLength.Num(), Vehicle->Wheels.Num());
+
 			int32 LengthCount = 0;
-			for (int32 WheelIdx = 0, NumWheels = Vehicle->Wheels.Num(); WheelIdx < NumWheels; ++WheelIdx)
+			for (int32 WheelIdx = 0; WheelIdx < NumWheels; ++WheelIdx)
 			{
 				Chaos::FSimpleSuspensionSim& Suspension = Vehicle->GetSuspension(WheelIdx);
 				Suspension.SetLastSpringLength(SuspensionLastSpringLength[WheelIdx]);
