Dstm: Chaos vehicle support 1.4: Avoid wheel creation in ChaosWheeledVehicleMovementComponent, so there won't be extra wheels during migration

#Synced-CL 48671814

diff --git a/Engine/Plugins/Experimental/ChaosVehiclesPlugin/Source/ChaosVehicles/Private/ChaosWheeledVehicleMovementComponent.cpp b/Engine/Plugins/Experimental/ChaosVehiclesPlugin/Source/ChaosVehicles/Private/ChaosWheeledVehicleMovementComponent.cpp
--- a/Engine/Plugins/Experimental/ChaosVehiclesPlugin/Source/ChaosVehicles/Private/ChaosWheeledVehicleMovementComponent.cpp
+++ b/Engine/Plugins/Experimental/ChaosVehiclesPlugin/Source/ChaosVehicles/Private/ChaosWheeledVehicleMovementComponent.cpp
@@ -39,6 +39,10 @@
 UE_DISABLE_OPTIMIZATION
 #endif
 
+#if UE_WITH_REMOTE_OBJECT_HANDLE
+#include "UObject/UObjectMigrationContext.h"
+#endif
+
 DECLARE_STATS_GROUP(TEXT("ChaosVehicle"), STATGROUP_ChaosVehicle, STATGROUP_Advanced);
 
 DECLARE_CYCLE_STAT(TEXT("Vehicle:SuspensionRaycasts"), STAT_ChaosVehicle_SuspensionRaycasts, STATGROUP_ChaosVehicle);
@@ -1138,8 +1142,33 @@ void UChaosWheeledVehicleMovementComponent::Serialize(FArchive & Ar)
 	Super::Serialize(Ar);
 
 	Ar.UsingCustomVersion(FUE5MainStreamObjectVersion::GUID);
+
+#if UE_WITH_REMOTE_OBJECT_HANDLE
+	if (Ar.IsMigratingRemoteObjects())
+	{
+		if (Ar.IsLoading())
+		{
+			// Check and verify we do receive wheels properly. Then they should not be recreated in CreateVehicle()
+			if (WheelSetups.Num() == Wheels.Num())
+			{
+				bAreWheelsMigrated = true;
+			}
+		}
+	}
+#endif
 }
 
+#if UE_WITH_REMOTE_OBJECT_HANDLE
+void UChaosWheeledVehicleMovementComponent::PostMigrate(const struct FUObjectMigrationContext& MigrationContext)
+{
+	Super::PostMigrate(MigrationContext);
+
+	// Set bAreWheelsMigrated back to false so it doesn't affect CreateWheels behavior. 
+	bAreWheelsMigrated = false;
+}
+#endif
+
+
 void UChaosWheeledVehicleMovementComponent::PostLoad()
 {
 	Super::PostLoad();
@@ -1371,17 +1400,25 @@ void UChaosWheeledVehicleMovementComponent::ComputeConstants()
 
 void UChaosWheeledVehicleMovementComponent::CreateWheels()
 {
-	// Wheels num is getting copied when blueprint recompiles, so we have to manually reset here
-	Wheels.Reset();
-
-	// Instantiate the wheels
-	for (int32 WheelIdx = 0; WheelIdx < WheelSetups.Num(); ++WheelIdx)
+#if UE_WITH_REMOTE_OBJECT_HANDLE
+	// If this function is called within PostMigrate for a migration, the wheels should be migrated, wheel creation is avoided. 
+	if (!bAreWheelsMigrated)
 	{
-		UChaosVehicleWheel* Wheel = NewObject<UChaosVehicleWheel>(this, WheelSetups[WheelIdx].WheelClass);
-		check(Wheel);
+#endif
+		// Wheels num is getting copied when blueprint recompiles, so we have to manually reset here
+		Wheels.Reset();
 
-		Wheels.Add(Wheel);
+		// Instantiate the wheels
+		for (int32 WheelIdx = 0; WheelIdx < WheelSetups.Num(); ++WheelIdx)
+		{
+			UChaosVehicleWheel* Wheel = NewObject<UChaosVehicleWheel>(this, WheelSetups[WheelIdx].WheelClass);
+			check(Wheel);
+
+			Wheels.Add(Wheel);
+		}
+#if UE_WITH_REMOTE_OBJECT_HANDLE
 	}
+#endif
 
 	// Initialize the wheels
 	for (int32 WheelIdx = 0; WheelIdx < Wheels.Num(); ++WheelIdx)
diff --git a/Engine/Plugins/Experimental/ChaosVehiclesPlugin/Source/ChaosVehicles/Public/ChaosWheeledVehicleMovementComponent.h b/Engine/Plugins/Experimental/ChaosVehiclesPlugin/Source/ChaosVehicles/Public/ChaosWheeledVehicleMovementComponent.h
--- a/Engine/Plugins/Experimental/ChaosVehiclesPlugin/Source/ChaosVehicles/Public/ChaosWheeledVehicleMovementComponent.h
+++ b/Engine/Plugins/Experimental/ChaosVehiclesPlugin/Source/ChaosVehicles/Public/ChaosWheeledVehicleMovementComponent.h
@@ -923,6 +923,10 @@ protected:
 		return WheelTrackDimensions;
 	}
 
+#if UE_WITH_REMOTE_OBJECT_HANDLE
+	ENGINE_API virtual void PostMigrate(const struct FUObjectMigrationContext& MigrationContext) override;
+#endif
+
 	private:
 
 
@@ -951,6 +955,11 @@ protected:
 	TArray<FWheelStatus> WheelStatus; /** Wheel output status */
 	TArray<FCachedState> CachedState;
 	Chaos::FPerformanceMeasure PerformanceMeasure;
+
+#if UE_WITH_REMOTE_OBJECT_HANDLE
+	bool bAreWheelsMigrated = false;
+#endif
+
 };
 
 #if VEHICLE_DEBUGGING_ENABLED
