Fix uninitialised field in FWheelsOutput structure

#Synced-CL 48719260

diff --git a/Engine/Plugins/Experimental/ChaosVehiclesPlugin/Source/ChaosVehicles/Public/ChaosVehicleManagerAsyncCallback.h b/Engine/Plugins/Experimental/ChaosVehiclesPlugin/Source/ChaosVehicles/Public/ChaosVehicleManagerAsyncCallback.h
--- a/Engine/Plugins/Experimental/ChaosVehiclesPlugin/Source/ChaosVehicles/Public/ChaosVehicleManagerAsyncCallback.h
+++ b/Engine/Plugins/Experimental/ChaosVehiclesPlugin/Source/ChaosVehicles/Public/ChaosVehicleManagerAsyncCallback.h
@@ -246,6 +246,7 @@ struct CHAOSVEHICLES_API FWheelsOutput
 		, DriveTorque(0.f)
 		, BrakeTorque(0.f)
 		, bABSActivated(false)
+		, bBlockingHit(false)
 		, ImpactPoint(FVector::ZeroVector)
 		, HitLocation(FVector::ZeroVector)
 		, PhysMaterial(nullptr)
