Dstm: Chaos vehicle support 2.2: Skip initial inertia scaling for migrated vehicles. Add migration for VehicleSimulation data.

#Synced-CL 51092234

diff --git a/Engine/Plugins/Experimental/ChaosVehiclesPlugin/Source/ChaosVehicles/Private/ChaosVehicleManagerAsyncCallback.cpp b/Engine/Plugins/Experimental/ChaosVehiclesPlugin/Source/ChaosVehicles/Private/ChaosVehicleManagerAsyncCallback.cpp
--- a/Engine/Plugins/Experimental/ChaosVehiclesPlugin/Source/ChaosVehicles/Private/ChaosVehicleManagerAsyncCallback.cpp
+++ b/Engine/Plugins/Experimental/ChaosVehiclesPlugin/Source/ChaosVehicles/Private/ChaosVehicleManagerAsyncCallback.cpp
@@ -31,6 +31,12 @@ void FChaosVehicleManagerAsyncCallback::ProcessInputs_Internal(int32 PhysicsStep
 
 	for (const TUniquePtr<FChaosVehicleAsyncInput>& VehicleInput : AsyncInput->VehicleInputs)
 	{
+#if UE_WITH_REMOTE_OBJECT_HANDLE
+		if (UChaosVehicleMovementComponent* VehicleComp = VehicleInput->Vehicle)
+		{
+			VehicleComp->ApplyMigratedVehicleState();
+		}
+#endif
 		UChaosVehicleSimulation* VehicleSim = VehicleInput->Vehicle->VehicleSimulationPT.Get();
 
 		if (VehicleSim == nullptr || !VehicleInput->Vehicle->bUsingNetworkPhysicsPrediction)
diff --git a/Engine/Plugins/Experimental/ChaosVehiclesPlugin/Source/ChaosVehicles/Private/ChaosVehicleMovementComponent.cpp b/Engine/Plugins/Experimental/ChaosVehiclesPlugin/Source/ChaosVehicles/Private/ChaosVehicleMovementComponent.cpp
--- a/Engine/Plugins/Experimental/ChaosVehiclesPlugin/Source/ChaosVehicles/Private/ChaosVehicleMovementComponent.cpp
+++ b/Engine/Plugins/Experimental/ChaosVehiclesPlugin/Source/ChaosVehicles/Private/ChaosVehicleMovementComponent.cpp
@@ -41,6 +41,9 @@
 #include "PhysicsProxy/SingleParticlePhysicsProxy.h"
 #include "PBDRigidsSolver.h"
 
+#if UE_WITH_REMOTE_OBJECT_HANDLE
+#include "UObject/UObjectMigrationContext.h"
+#endif
 
 #define LOCTEXT_NAMESPACE "UVehicleMovementComponent"
 
@@ -670,6 +673,10 @@ void UChaosVehicleMovementComponent::Serialize(FArchive& Ar)
 #if UE_WITH_REMOTE_OBJECT_HANDLE
 	if (Ar.IsMigratingRemoteObjects())
 	{
+		if (Ar.IsLoading())
+		{
+			bIsMigrated = true;
+		}
 		MigratePhysicsState(Ar);
 	}
 #endif
@@ -690,8 +697,77 @@ void UChaosVehicleMovementComponent::MigratePhysicsState(FArchive& Ar)
 	{
 		FPhysicsVehicleOutput::StaticStruct()->SerializeItem(Ar, PVehicleOutput.Get(), nullptr);
 	}
+
+	if (Ar.IsLoading())
+	{
+		bIsVehicleStatesMigrated = true;
+		MigratedVehicleStates.NetSerialize(Ar, nullptr, bIsVehicleStatesMigrated);
+	}
+	else if (Ar.IsSaving())
+	{
+		const EObjectMigrationSendType ObjMigrationSendType = MigrationContext->GetObjectMigrationSendType(this);
+
+		if (CanSendMigratedPhysicsState(MigrationContext, ObjMigrationSendType))
+		{
+			bool bOutSuccess = true;
+			if (bIsVehicleStatesMigrated)
+			{
+				UE_LOGF(LogVehicle, Verbose, "[%hs] Migrated a pending vehicle state", __func__);
+				MigratedVehicleStates.NetSerialize(Ar, nullptr, bOutSuccess);
+			}
+			else
+			{
+				FNetworkVehicleStates VehicleStates;
+				// BuildData accesses VehicleSimulationPT, which contains Physics Thread data.
+				// While FPhysicsCommand is not inherently thread-safe due to PT ownership, 
+				// this is currently safe because Dstm enforces synchronous physics execution.
+				if (FPhysScene* PhysScene = GetWorld() ? GetWorld()->GetPhysicsScene() : nullptr)
+				{
+					FPhysicsCommand::ExecuteWrite(PhysScene, [this, &VehicleStates]()
+						{
+							VehicleStates.BuildData(this);
+						});
+				}
+				else
+				{
+					UE_LOGF(LogVehicle, Warning, "[%hs] Cannot BuildData - no valid physics scene", __func__);
+				}
+				VehicleStates.NetSerialize(Ar, nullptr, bOutSuccess);
+			}
+		}
+		else
+		{
+			UE_LOGF(LogVehicle, Log, "[%hs] Migrated an empty vehicle physics state", __func__);
+			FNetworkVehicleStates VehicleStates;
+			bool bOutSuccess = true;
+			VehicleStates.NetSerialize(Ar, nullptr, bOutSuccess);
+		}
+	}
 }
 
+bool UChaosVehicleMovementComponent::CanSendMigratedPhysicsState(const FUObjectMigrationContext* MigrationContext, EObjectMigrationSendType ObjMigrationSendType) const
+{
+	if (ObjMigrationSendType == EObjectMigrationSendType::ReturnBorrowed)
+	{
+		return MigrationContext->MultiServerCommitRequestId.IsValid();
+	}
+
+	return true;
+}
+
+void UChaosVehicleMovementComponent::PostMigrate(const struct FUObjectMigrationContext& MigrationContext)
+{
+	Super::PostMigrate(MigrationContext);
+}
+
+void UChaosVehicleMovementComponent::ApplyMigratedVehicleState()
+{
+	if (bIsVehicleStatesMigrated)
+	{
+		MigratedVehicleStates.ApplyData(this);
+		bIsVehicleStatesMigrated = false;
+	}
+}
 #endif 
 
 #if WITH_EDITOR
@@ -1522,8 +1598,13 @@ void UChaosVehicleMovementComponent::SetupVehicleMass()
 	{
 		//Ensure that if mass properties ever change we set them back to our override
 		UpdatedPrimitive->GetBodyInstance()->OnRecalculatedMassProperties().AddUObject(this, &UChaosVehicleMovementComponent::UpdateMassProperties);
-
-		UpdateMassProperties(UpdatedPrimitive->GetBodyInstance());
+#if UE_WITH_REMOTE_OBJECT_HANDLE
+		// Skip inertia scaling when vehicle is migrated since its scaled inertia is already migrated. 
+		if (!bIsMigrated)
+#endif
+		{
+			UpdateMassProperties(UpdatedPrimitive->GetBodyInstance());
+		}
 	}
 }
 
@@ -1651,7 +1732,13 @@ void UChaosVehicleMovementComponent::DrawLine2D(UCanvas* Canvas, const FVector2D
 
 TUniquePtr<Chaos::FSimpleWheeledVehicle> UChaosVehicleMovementComponent::CreatePhysicsVehicle()
 {
-	PVehicleOutput = MakeUnique<FPhysicsVehicleOutput>();	// create physics output container
+#if UE_WITH_REMOTE_OBJECT_HANDLE
+	// Don't reinitialize PVehicleOutput when migrated because its data is already migrated over. 
+	if (!bIsMigrated)
+#endif
+	{
+		PVehicleOutput = MakeUnique<FPhysicsVehicleOutput>();	// create physics output container
+	}
 	return MakeUnique<Chaos::FSimpleWheeledVehicle>();		// create physics sim
 }
 
diff --git a/Engine/Plugins/Experimental/ChaosVehiclesPlugin/Source/ChaosVehicles/Public/ChaosVehicleMovementComponent.h b/Engine/Plugins/Experimental/ChaosVehiclesPlugin/Source/ChaosVehicles/Public/ChaosVehicleMovementComponent.h
--- a/Engine/Plugins/Experimental/ChaosVehiclesPlugin/Source/ChaosVehicles/Public/ChaosVehicleMovementComponent.h
+++ b/Engine/Plugins/Experimental/ChaosVehiclesPlugin/Source/ChaosVehicles/Public/ChaosVehicleMovementComponent.h
@@ -1155,6 +1155,9 @@ public:
 	UPROPERTY(EditAnywhere, Category = VehicleInput, AdvancedDisplay)
 	FVehicleInputRateConfig YawInputRate;
 
+#if UE_WITH_REMOTE_OBJECT_HANDLE
+	void ApplyMigratedVehicleState();
+#endif
 protected:
 
 	// input related
@@ -1262,8 +1265,12 @@ protected:
 	void DrawLine2D(UCanvas* Canvas, const FVector2D& StartPos, const FVector2D& EndPos, FColor Color, float Thickness = 1.f);
 
 #if UE_WITH_REMOTE_OBJECT_HANDLE
+	ENGINE_API bool CanSendMigratedPhysicsState(const FUObjectMigrationContext* MigrationContext, EObjectMigrationSendType ObjMigrationSendType) const;
+
 	// migrates physics related state (such as vehicle ouput)
 	ENGINE_API virtual void MigratePhysicsState(FArchive& Ar);
+
+	ENGINE_API virtual void PostMigrate(const struct FUObjectMigrationContext& MigrationContext) override;
 #endif
 
 	float GetForwardAcceleration()
@@ -1309,4 +1316,11 @@ protected:
 	float PrevSteeringInput;
 	float PrevReplicatedSteeringInput;
 	bool bUsingNetworkPhysicsPrediction;
+
+private:
+#if UE_WITH_REMOTE_OBJECT_HANDLE
+	bool bIsMigrated = false;
+	FNetworkVehicleStates MigratedVehicleStates;
+	bool bIsVehicleStatesMigrated = false;
+#endif
 };
