[ChaosVehicle] Fixes
- Fix issues with transmission where it could get stuck in reverse while pressing forward or start occilating between forward and reverse when pressing backwards input.
- Fix wheel flickering with async physics

#Synced-CL 52337891

diff --git a/Engine/Plugins/Experimental/ChaosVehiclesPlugin/Source/ChaosVehicles/Private/ChaosVehicleManagerAsyncCallback.cpp b/Engine/Plugins/Experimental/ChaosVehiclesPlugin/Source/ChaosVehicles/Private/ChaosVehicleManagerAsyncCallback.cpp
--- a/Engine/Plugins/Experimental/ChaosVehiclesPlugin/Source/ChaosVehicles/Private/ChaosVehicleManagerAsyncCallback.cpp
+++ b/Engine/Plugins/Experimental/ChaosVehiclesPlugin/Source/ChaosVehicles/Private/ChaosVehicleManagerAsyncCallback.cpp
@@ -57,8 +57,16 @@ void FChaosVehicleManagerAsyncCallback::ProcessInputs_Internal(int32 PhysicsStep
 
 		APlayerController* PlayerController = VehicleInput->Vehicle->GetPlayerController();
 		if(PlayerController && PlayerController->IsLocalController() && !bIsResimming)
-		{ 
+		{
 			VehicleSim->VehicleInputs = VehicleInput->PhysicsInputs.NetworkInputs.VehicleInputs;
+
+			// Apply the GT's gear decision to the PT transmission so it's authoritative
+			// for the current frame and gets captured correctly by BuildData
+			if (VehicleSim->PVehicle && VehicleSim->PVehicle->HasTransmission())
+			{
+				Chaos::FSimpleTransmissionSim& Transmission = VehicleSim->PVehicle->GetTransmission();
+				Transmission.SetGear(VehicleInput->PhysicsInputs.NetworkInputs.TransmissionTargetGear, true);
+			}
 		}
 		else
 		{
@@ -182,6 +190,11 @@ bool FNetworkVehicleInputs::NetSerialize(FArchive& Ar, class UPackageMap* Map, b
 	Ar << VehicleInputs.YawInput;
 	Ar << VehicleInputs.HandbrakeInput;
 
+	Ar << VehicleInputs.ParkingEnabled;
+	Ar << VehicleInputs.GearUpInput;
+	Ar << VehicleInputs.GearDownInput;
+	Ar << VehicleInputs.TransmissionType;
+
 	Ar << TransmissionChangeTime;
 	Ar << TransmissionCurrentGear;
 	Ar << TransmissionTargetGear;
diff --git a/Engine/Plugins/Experimental/ChaosVehiclesPlugin/Source/ChaosVehicles/Private/ChaosVehicleMovementComponent.cpp b/Engine/Plugins/Experimental/ChaosVehiclesPlugin/Source/ChaosVehicles/Private/ChaosVehicleMovementComponent.cpp
--- a/Engine/Plugins/Experimental/ChaosVehiclesPlugin/Source/ChaosVehicles/Private/ChaosVehicleMovementComponent.cpp
+++ b/Engine/Plugins/Experimental/ChaosVehiclesPlugin/Source/ChaosVehicles/Private/ChaosVehicleMovementComponent.cpp
@@ -1048,19 +1048,27 @@ void UChaosVehicleMovementComponent::SetChangeDownInput(bool bNewGearDown)
 
 void UChaosVehicleMovementComponent::SetTargetGear(int32 GearNum, bool bImmediate)
 {
-	if (PVehicleOutput && GearNum != PVehicleOutput->TargetGear)
+	if (GearNum != TargetGear)
 	{
-		FBodyInstance* TargetInstance = UpdatedPrimitive->GetBodyInstance();
-
-		if (TargetInstance)
+		if (bUsingNetworkPhysicsPrediction)
 		{
-			FPhysicsCommand::ExecuteWrite(TargetInstance->GetPhysicsActor(), [&](const FPhysicsActorHandle& Chassis)
-				{
-					if (VehicleSimulationPT && VehicleSimulationPT->PVehicle && VehicleSimulationPT->PVehicle->HasTransmission())
+			// With network physics prediction, the gear is set via the input data path
+			// (Update -> ProcessInputs_Internal) to stay consistent with the input history.
+			// Direct ExecuteWrite would conflict with ApplyData during resimulation.
+		}
+		else
+		{
+			// Without network physics prediction, set the gear directly on the PT.
+			if (FBodyInstance* TargetInstance = GetBodyInstance())
+			{
+				FPhysicsCommand::ExecuteWrite(TargetInstance->GetPhysicsActor(), [&](const FPhysicsActorHandle& Chassis)
 					{
-						VehicleSimulationPT->PVehicle->GetTransmission().SetGear(GearNum, bImmediate);
-					}
-				});
+						if (VehicleSimulationPT && VehicleSimulationPT->PVehicle && VehicleSimulationPT->PVehicle->HasTransmission())
+						{
+							VehicleSimulationPT->PVehicle->GetTransmission().SetGear(GearNum, bImmediate);
+						}
+					});
+			}
 		}
 
 		TargetGear = GearNum;
@@ -1148,11 +1156,6 @@ void UChaosVehicleMovementComponent::CalcThrottleBrakeInput(float& ThrottleOut,
 				BrakeOut = 0.0f;
 			}
 		}
-		// straight reversing
-		else if (RawBrakeInput > 0.f && GetTargetGear() < 0)
-		{
-			ThrottleOut = RawBrakeInput;
-		}
 		else
 		{
 			// if player isn't pressing forward or backwards...
@@ -1283,18 +1286,18 @@ void UChaosVehicleMovementComponent::UpdateState(float DeltaTime)
 	{
 		if (bReverseAsBrake)
 		{
-			//for reverse as state we want to automatically shift between reverse and first gear
-			// Note: Removed this condition to support wheel spinning when rolling backwards with accelerator pressed, rather than braking
-			//if (FMath::Abs(GetForwardSpeed()) < WrongDirectionThreshold)	//we only shift between reverse and first if the car is slow enough.
+			// Automatically shift between reverse and first gear based on input.
+			// Throttle takes priority when both inputs are pressed to avoid gear oscillation.
+			if (RawThrottleInput > KINDA_SMALL_NUMBER && GetTargetGear() <= 0)
+			{
+				SetTargetGear(1, true);
+			}
+			else if (RawBrakeInput > KINDA_SMALL_NUMBER && RawThrottleInput <= KINDA_SMALL_NUMBER && GetTargetGear() >= 0)
 			{
-				if (RawBrakeInput > KINDA_SMALL_NUMBER && GetCurrentGear() >= 0 && GetTargetGear() >= 0)
+				if (FMath::Abs(GetForwardSpeed()) < WrongDirectionThreshold)
 				{
 					SetTargetGear(-1, true);
 				}
-				else if (RawThrottleInput > KINDA_SMALL_NUMBER && GetCurrentGear() <= 0 && GetTargetGear() <= 0)
-				{
-					SetTargetGear(1, true);
-				}
 			}
 		}
 		else
@@ -1884,6 +1887,11 @@ void UChaosVehicleMovementComponent::Update(float DeltaTime)
 				AsyncInput->PhysicsInputs.NetworkInputs.VehicleInputs.TransmissionType = TransmissionType;
 				AsyncInput->PhysicsInputs.NetworkInputs.VehicleInputs.ParkingEnabled = bParkEnabled;
 
+				// Pass the GT's gear decision through the input data path so it flows into
+				// the network physics input history via BuildData/ApplyData consistently
+				AsyncInput->PhysicsInputs.NetworkInputs.TransmissionTargetGear = TargetGear;
+				AsyncInput->PhysicsInputs.NetworkInputs.TransmissionCurrentGear = TargetGear;
+
 				// debug feature to limit the vehicles top speed
 				if ((GVehicleDebugParams.SetMaxMPH > 0.f) && (FMath::Abs(ThrottleInput) > 0.0f) && FMath::Abs(GetForwardSpeedMPH()) >= GVehicleDebugParams.SetMaxMPH)
 				{
diff --git a/Engine/Plugins/Experimental/ChaosVehiclesPlugin/Source/ChaosVehicles/Private/ChaosWheeledVehicleMovementComponent.cpp b/Engine/Plugins/Experimental/ChaosVehiclesPlugin/Source/ChaosVehicles/Private/ChaosWheeledVehicleMovementComponent.cpp
--- a/Engine/Plugins/Experimental/ChaosVehiclesPlugin/Source/ChaosVehicles/Private/ChaosWheeledVehicleMovementComponent.cpp
+++ b/Engine/Plugins/Experimental/ChaosVehiclesPlugin/Source/ChaosVehicles/Private/ChaosWheeledVehicleMovementComponent.cpp
@@ -875,7 +875,14 @@ void UChaosWheeledVehicleSimulation::ProcessMechanicalSimulation(float DeltaTime
 		PEngine.Simulate(DeltaTime);
 
 		PTransmission.SetEngineRPM(PEngine.GetEngineRPM()); // needs engine RPM to decide when to change gear (automatic gearbox)
-		PTransmission.SetAllowedToChangeGear(!VehicleState.bVehicleInAir && !IsWheelSpinning());
+		if (Chaos::FPhysicsSolverBase::IsNetworkPhysicsPredictionEnabled())
+		{
+			PTransmission.SetAllowedToChangeGear(true);
+		}
+		else
+		{
+			PTransmission.SetAllowedToChangeGear(!VehicleState.bVehicleInAir && !IsWheelSpinning());
+		}
 		float GearRatio = PTransmission.GetGearRatio(PTransmission.GetCurrentGear());
 
 		PTransmission.Simulate(DeltaTime);
@@ -2846,18 +2853,24 @@ float UChaosWheeledVehicleMovementComponent::GetSuspensionOffset(int WheelIndex)
 			{
 				if (WheelStatus[WheelIndex].bInContact)
 				{
-					FVector LocalPos = GetWheelRestingPosition(WheelSetup);
-
-					FVector ReferencePos = (Wheel->SweepShape == ESweepShape::Spherecast)? WheelStatus[WheelIndex].HitLocation : WheelStatus[WheelIndex].ContactPoint;
-					FVector LocalHitPoint = VehicleWorldTransform.InverseTransformPosition(ReferencePos);
-					float Radius = PVehicleOutput->Wheels[WheelIndex].WheelRadius;
+					float NewOffset;
+					if (PVehicleOutput)
+					{
+						// If we have output from the async simulation, use that SuspensionOffset rather than recalculating from interpolated contact points.
+						NewOffset = PVehicleOutput->Wheels[WheelIndex].SuspensionOffset;
+					}
+					else
+					{
+						FVector LocalPos = GetWheelRestingPosition(WheelSetup);
+						FVector ReferencePos = (Wheel->SweepShape == ESweepShape::Spherecast) ? WheelStatus[WheelIndex].HitLocation : WheelStatus[WheelIndex].ContactPoint;
+						FVector LocalHitPoint = VehicleWorldTransform.InverseTransformPosition(ReferencePos);
+						NewOffset = CalcOffset(Wheel->SweepShape, LocalHitPoint, LocalPos, Wheel->WheelRadius);
+					}
 
 					if (CachedState[WheelIndex].bIsValid)
 					{
 						Offset = CachedState[WheelIndex].WheelOffset;
 
-						float NewOffset = CalcOffset(Wheel->SweepShape, LocalHitPoint, LocalPos, Radius);
-							
 						// interpolate between old and new positions or will just jump to new position if Wheel->SuspensionSmoothing == 0
 						float InterpolationMultiplier = 1.0f - (Wheel->SuspensionSmoothing / 11.0f);
 						float Delta = NewOffset - CachedState[WheelIndex].WheelOffset;
@@ -2865,7 +2878,7 @@ float UChaosWheeledVehicleMovementComponent::GetSuspensionOffset(int WheelIndex)
 					}
 					else
 					{
-						Offset = CalcOffset(Wheel->SweepShape, LocalHitPoint, LocalPos, Radius);
+						Offset = NewOffset;
 					}
 					Offset = FMath::Clamp(Offset, -Wheel->SuspensionMaxDrop, Wheel->SuspensionMaxRaise);
 				}
